[NNE] Fix NoPCH compile error

#rnx

diff --git a/Engine/Plugins/Experimental/NNERuntimeIREE/Source/NNERuntimeIREE/Private/NNERuntimeIREECommon.h b/Engine/Plugins/Experimental/NNERuntimeIREE/Source/NNERuntimeIREE/Private/NNERuntimeIREECommon.h
--- a/Engine/Plugins/Experimental/NNERuntimeIREE/Source/NNERuntimeIREE/Private/NNERuntimeIREECommon.h
+++ b/Engine/Plugins/Experimental/NNERuntimeIREE/Source/NNERuntimeIREE/Private/NNERuntimeIREECommon.h
@@ -6,6 +6,7 @@
 
 #include "Interfaces/ITargetPlatform.h"
 #include "Kismet/GameplayStatics.h"
+#include "Misc/Paths.h"
 #include "Misc/ScopeLock.h"
 #include "NNEModelData.h"
 #include "NNERuntimeCPU.h"
diff --git a/Engine/Plugins/Experimental/NNERuntimeIREE/Source/NNERuntimeIREE/Private/NNERuntimeIREECpu.cpp b/Engine/Plugins/Experimental/NNERuntimeIREE/Source/NNERuntimeIREE/Private/NNERuntimeIREECpu.cpp
--- a/Engine/Plugins/Experimental/NNERuntimeIREE/Source/NNERuntimeIREE/Private/NNERuntimeIREECpu.cpp
+++ b/Engine/Plugins/Experimental/NNERuntimeIREE/Source/NNERuntimeIREE/Private/NNERuntimeIREECpu.cpp
@@ -12,6 +12,7 @@
 #include "IO/IoHash.h"
 #include "Kismet/GameplayStatics.h"
 #include "Misc/FileHelper.h"
+#include "Misc/Paths.h"
 #include "Misc/SecureHash.h"
 #include "NNERuntimeIREECommon.h"
 #include "NNERuntimeIREECpuModel.h"
