[NNE] Changed staging to use the configs instead of predefined paths

Problem: IREE produces shared libraries in the cooking which are not know upfront and have to be staged dynamically.

Solution: instead of compiling everything into one lib and try to stage it through a fixed path, we add an additional folder to be always staged to project settings.

Tested: Win, Linux, consoles, preflight, non-unity build and presubmit

#rnx

diff --git a/Engine/Plugins/Experimental/NNERuntimeIREE/NNERuntimeIREE.uplugin b/Engine/Plugins/Experimental/NNERuntimeIREE/NNERuntimeIREE.uplugin
--- a/Engine/Plugins/Experimental/NNERuntimeIREE/NNERuntimeIREE.uplugin
+++ b/Engine/Plugins/Experimental/NNERuntimeIREE/NNERuntimeIREE.uplugin
@@ -25,13 +25,5 @@
 			"Type": "Editor",
 			"LoadingPhase": "Default"
 		}
-	],
-	"PreBuildSteps": {
-		"Win64": [
-			"if exist \"$(PluginDir)\\Intermediate\\Build\\ObjectFiles\\\" rmdir /S /Q \"$(PluginDir)\\Intermediate\\Build\\ObjectFiles\""
-		],
-		"Linux": [
-			"if [ -d \"$(PluginDir)/Intermediate/Build/ObjectFiles\" ]; then rm -R -f \"$(PluginDir)/Intermediate/Build/ObjectFiles\" ; fi"
-		]
-	}
+	]
 }
\ No newline at end of file
diff --git a/Engine/Plugins/Experimental/NNERuntimeIREE/Source/NNERuntimeIREE/NNERuntimeIREE.Build.cs b/Engine/Plugins/Experimental/NNERuntimeIREE/Source/NNERuntimeIREE/NNERuntimeIREE.Build.cs
--- a/Engine/Plugins/Experimental/NNERuntimeIREE/Source/NNERuntimeIREE/NNERuntimeIREE.Build.cs
+++ b/Engine/Plugins/Experimental/NNERuntimeIREE/Source/NNERuntimeIREE/NNERuntimeIREE.Build.cs
@@ -16,42 +16,37 @@ public class NNERuntimeIREE : ModuleRules
 		PrivateDefinitions.Add("NNE_RUNTIME_IREE_PLATFORM_NAME=\"" + PlatformName + "\"");
 		PrivateDefinitions.Add("NNE_RUNTIME_IREE_PLATFORM_DISPLAY_NAME=\"" + PlatformDisplayName + "\"");
 		PrivateDefinitions.Add("NNE_RUNTIME_IREE_PLATFORM_SHARED_LIB_EXTENSION=\"" + SharedLibExtension + "\"");
-
-		// Use the non-remapping RuntimeDependencies.Add function so that the file does not get copied immediately but at staging time
-		string SharedLibPath = Path.Combine(PluginDirectory, "Binaries", PlatformName, SharedLibName + "." + SharedLibExtension);
-		Directory.CreateDirectory(Path.GetDirectoryName(SharedLibPath));
-		File.Create(SharedLibPath).Close();
-		RuntimeDependencies.Add(SharedLibPath, StagedFileType.NonUFS);
 	}
 
 	public NNERuntimeIREE(ReadOnlyTargetRules Target) : base(Target)
 	{
 		PCHUsage = PCHUsageMode.UseExplicitOrSharedPCHs;
 
-		PublicDependencyModuleNames.Add("Core");
-
 		PrivateDependencyModuleNames.AddRange
 		(
 			new string[]
 			{
+				"Core",
 				"CoreUObject",
 				"Engine",
 				"IREE",
-				"Json",
-				"NNE",
-				"Slate",
-				"SlateCore",
-				"Projects"
+				"NNE"
 			}
 		);
 
 		if (Target.Type == TargetType.Editor)
 		{
-			PrivateDependencyModuleNames.Add("DerivedDataCache");
+			PrivateDependencyModuleNames.AddRange(
+				new string[]
+				{
+					"DerivedDataCache",
+					"Json",
+					"Projects",
+					"TargetPlatform"
+				}
+			);
 		}
 
-		PrivateDefinitions.Add("NNE_RUNTIME_IREE_SHARED_LIB_NAME=\"" + SharedLibName + "\"");
-
 		if (Target.Platform == UnrealTargetPlatform.Win64)
 		{
 			SetupNNERuntimeIREE("Win64", "Windows", "dll");
diff --git a/Engine/Plugins/Experimental/NNERuntimeIREE/Source/NNERuntimeIREE/Private/NNERuntimeIREECommon.cpp b/Engine/Plugins/Experimental/NNERuntimeIREE/Source/NNERuntimeIREE/Private/NNERuntimeIREECommon.cpp
--- a/Engine/Plugins/Experimental/NNERuntimeIREE/Source/NNERuntimeIREE/Private/NNERuntimeIREECommon.cpp
+++ b/Engine/Plugins/Experimental/NNERuntimeIREE/Source/NNERuntimeIREE/Private/NNERuntimeIREECommon.cpp
@@ -151,7 +151,7 @@ namespace UE::NNERuntimeIREE
 
 	FString GetTargetPlatformDisplayName(const ITargetPlatform* TargetPlatform)
 	{
-		return TargetPlatform ? TargetPlatform->DisplayName().ToString() : UGameplayStatics::GetPlatformName();
+		return TargetPlatform ? TargetPlatform->IniPlatformName() : UGameplayStatics::GetPlatformName();
 	}
 
 	class FIREEInstance
diff --git a/Engine/Plugins/Experimental/NNERuntimeIREE/Source/NNERuntimeIREE/Private/NNERuntimeIREECpu.cpp b/Engine/Plugins/Experimental/NNERuntimeIREE/Source/NNERuntimeIREE/Private/NNERuntimeIREECpu.cpp
--- a/Engine/Plugins/Experimental/NNERuntimeIREE/Source/NNERuntimeIREE/Private/NNERuntimeIREECpu.cpp
+++ b/Engine/Plugins/Experimental/NNERuntimeIREE/Source/NNERuntimeIREE/Private/NNERuntimeIREECpu.cpp
@@ -7,7 +7,6 @@
 #include "EngineAnalytics.h"
 #include "GenericPlatform/GenericPlatformMisc.h"
 #include "HAL/PlatformFileManager.h"
-#include "Interfaces/IPluginManager.h"
 #include "Interfaces/ITargetPlatform.h"
 #include "IO/IoHash.h"
 #include "Kismet/GameplayStatics.h"
@@ -24,8 +23,10 @@
 #include "DerivedDataCache.h"
 #include "DerivedDataRequestOwner.h"
 #include "Dom/JsonObject.h"
+#include "Interfaces/IPluginManager.h"
 #include "Internationalization/TextLocalizationResource.h"
 #include "Memory/SharedBuffer.h"
+#include "Misc/App.h"
 #include "Serialization/JsonReader.h"
 #include "Serialization/JsonSerializer.h"
 #include "Serialization/JsonSerializerMacros.h"
@@ -170,8 +171,7 @@ namespace UE::NNERuntimeIREECpu::Private
 		{
 			FPaths::Combine(FPaths::ConvertRelativePathToFull(*IPluginManager::Get().FindPlugin(UE_PLUGIN_NAME)->GetBaseDir()), "Config", InBuildConfigFileName),
 			FPaths::Combine(FPaths::ConvertRelativePathToFull(*FPaths::EngineDir()), "Platforms", InTargetPlatformDisplayName, "Plugins", UE_PLUGIN_NAME, "Config", InBuildConfigFileName),
-			FPaths::Combine(FPaths::ConvertRelativePathToFull(*FPaths::EngineDir()), "Platforms", InTargetPlatformDisplayName, "Plugins", "Experimental", UE_PLUGIN_NAME, "Config", InBuildConfigFileName),
-			FPaths::Combine(FPaths::ConvertRelativePathToFull(*FPaths::ProjectDir()), "Platforms", InTargetPlatformDisplayName, "Plugins", UE_PLUGIN_NAME, "Config", InBuildConfigFileName)
+			FPaths::Combine(FPaths::ConvertRelativePathToFull(*FPaths::EngineDir()), "Platforms", InTargetPlatformDisplayName, "Plugins", "Experimental", UE_PLUGIN_NAME, "Config", InBuildConfigFileName)
 		};
 
 		FString BuildConfigFileString = "";
@@ -217,26 +217,22 @@ namespace UE::NNERuntimeIREECpu::Private
 	}
 #endif // WITH_EDITOR
 
-	FString GetIntermediateBuildDirPath()
-	{
-		FString PluginDir = FPaths::ConvertRelativePathToFull(*IPluginManager::Get().FindPlugin(UE_PLUGIN_NAME)->GetBaseDir());
-		return FPaths::Combine(PluginDir, "Intermediate", "Build");
-	}
-
-	FString GetIntermediateModeslDirPath(const FString& PlatformName)
+	FString GetIntermediateModelDirPath(const FString& PlatformName, const FString& FileIdString)
 	{
-		return FPaths::Combine(GetIntermediateBuildDirPath(), PlatformName, "Models");
+		FString IntermediateDir = FPaths::ConvertRelativePathToFull(*FPaths::ProjectIntermediateDir());
+		return FPaths::Combine(IntermediateDir, "Build", PlatformName, UE_PLUGIN_NAME, FileIdString);
 	}
 
-	FString GetIntermediateModelDirPath(const FString& PlatformName, const FString& FileIdString)
+	FString GetCookedModelDirPath(const FString& PlatformDisplayName)
 	{
-		return FPaths::Combine(GetIntermediateModeslDirPath(PlatformName), FileIdString);
+		FString SavedDir = FPaths::ConvertRelativePathToFull(*FPaths::ProjectSavedDir());
+		return FPaths::Combine(SavedDir, "Cooked", PlatformDisplayName, "Engine", "Plugins", UE_PLUGIN_NAME, "Binaries");
 	}
 
-	FString GetAssembledSharedLibDirPath(const FString& PlatformName)
+	FString GetPackagedModelDirPath(const FString& PlatformName)
 	{
-		FString PluginDir = FPaths::ConvertRelativePathToFull(*IPluginManager::Get().FindPlugin(UE_PLUGIN_NAME)->GetBaseDir());
-		return FPaths::Combine(PluginDir, "Binaries", PlatformName);
+		FString ProjectDir = FPaths::ConvertRelativePathToFull(*FPaths::ProjectDir());
+		return FPaths::Combine(ProjectDir, "Binaries", PlatformName, UE_PLUGIN_NAME);
 	}
 
 	FString GetModelCpuDataIdentifier(const FString& FileIdString, const FString& PlatformDisplayName)
@@ -358,47 +354,17 @@ TSharedPtr<UE::NNE::FSharedModelData> UNNERuntimeIREECpu::CreateModelData(FStrin
 			return TSharedPtr<UE::NNE::FSharedModelData>();
 		}
 
+		FString CookedModelPath = FPaths::Combine(UE::NNERuntimeIREECpu::Private::GetCookedModelDirPath(TargetPlatformDisplayName), SharedLibName);
+		IFileManager::Get().Copy(*CookedModelPath, *SharedLibPath);
+
 		TArray<uint8> SharedLibData;
-		if (!FFileHelper::LoadFileToArray(SharedLibData, *SharedLibPath))
+		if (!FFileHelper::LoadFileToArray(SharedLibData, *CookedModelPath) || SharedLibData.Num() < 1)
 		{
-			UE_LOG(LogTemp, Warning, TEXT("UNNERuntimeIREECpu could not read the shared library \"%s\""), *SharedLibPath);
+			UE_LOG(LogTemp, Warning, TEXT("UNNERuntimeIREECpu could not read the shared library \"%s\""), *CookedModelPath);
 			return TSharedPtr<UE::NNE::FSharedModelData>();
 		}
 		FSharedBuffer SharedBuffer = MakeSharedBufferFromArray(MoveTemp(SharedLibData));
 		UE::NNERuntimeIREECpu::Private::PutIntoDDC(FileId, GetRuntimeName(), UE::NNERuntimeIREECpu::Private::GetModelCpuDataIdentifier(FileIdString, TargetPlatformDisplayName), SharedBuffer);
-
-		// Scoped file lock
-		{
-			TUniquePtr<IFileHandle> LockFileHandle;
-			while (!(LockFileHandle = TUniquePtr<IFileHandle>(PlatformFile.OpenWrite(*FPaths::Combine(UE::NNERuntimeIREECpu::Private::GetIntermediateModeslDirPath(BuildConfig.TargetPlatformName), ".lock")))).IsValid())
-			{
-				FPlatformProcess::Sleep(0.005);
-			}
-			checkf(!TUniquePtr<IFileHandle>(PlatformFile.OpenWrite(*FPaths::Combine(UE::NNERuntimeIREECpu::Private::GetIntermediateModeslDirPath(BuildConfig.TargetPlatformName), ".lock"))).IsValid(), TEXT("NNERuntimeIREE file lock failed"));
-
-			IFileManager::Get().Copy(*FPaths::Combine(UE::NNERuntimeIREECpu::Private::GetIntermediateBuildDirPath(), "ObjectFiles", FileIdString + ".o"), *ObjectPath);
-			FJsonSerializableArray FilePaths;
-			PlatformFile.FindFiles(FilePaths, *FPaths::Combine(UE::NNERuntimeIREECpu::Private::GetIntermediateBuildDirPath(), "ObjectFiles"), TEXT(".o"));
-			FString ObjectPaths = "";
-			for (int32 i = 0; i < FilePaths.Num(); i++)
-			{
-				ObjectPaths += FilePaths[i];
-				if (i < FilePaths.Num() - 1)
-				{
-					ObjectPaths += "\" \"";
-				}
-			}
-			FString CompleteSharedLibPath = FPaths::Combine(UE::NNERuntimeIREECpu::Private::GetAssembledSharedLibDirPath(BuildConfig.TargetPlatformName), FString(NNE_RUNTIME_IREE_SHARED_LIB_NAME) + "." + BuildConfig.SharedLibraryExtension);
-			FString CompleteFinalSharedLibraryCompilationArguments = SharedLibraryCompilationArguments.Replace(*FString("${OBJECT_PATH}"), *ObjectPaths).Replace(*FString("${SHARED_LIB_PATH}"), *CompleteSharedLibPath);
-			FProcHandle CompleteSharedLibraryCompilationProcHandle = FPlatformProcess::CreateProc(*SharedLibraryCompilationCommand, *CompleteFinalSharedLibraryCompilationArguments, false, true, true, nullptr, 0, nullptr, nullptr, nullptr);
-			FPlatformProcess::WaitForProc(CompleteSharedLibraryCompilationProcHandle);
-			if (!PlatformFile.FileExists(*CompleteSharedLibPath))
-			{
-				UE_LOG(LogTemp, Warning, TEXT("UNNERuntimeIREECpu failed to compile the shared library for the objects \"%s\" using the command:"), *ObjectPaths);
-				UE_LOG(LogTemp, Warning, TEXT("\"%s\" %s"), *SharedLibraryCompilationCommand, *CompleteFinalSharedLibraryCompilationArguments);
-				return TSharedPtr<UE::NNE::FSharedModelData>();
-			}
-		}
 	}
 
 	FString HeaderPath = FPaths::Combine(IntermediateModelDir, FileIdString + ".h");
@@ -532,8 +498,8 @@ TSharedPtr<UE::NNE::IModelCPU> UNNERuntimeIREECpu::CreateModelCPU(TObjectPtr<UNN
 	FString SharedLibDirPath = UE::NNERuntimeIREECpu::Private::GetIntermediateModelDirPath(NNE_RUNTIME_IREE_PLATFORM_NAME, FileIdString);
 	FString SharedLibName = FileIdString + "." + NNE_RUNTIME_IREE_PLATFORM_SHARED_LIB_EXTENSION;
 #else
-	FString SharedLibDirPath = UE::NNERuntimeIREECpu::Private::GetAssembledSharedLibDirPath(NNE_RUNTIME_IREE_PLATFORM_NAME);
-	FString SharedLibName = FString(NNE_RUNTIME_IREE_SHARED_LIB_NAME) + "." + NNE_RUNTIME_IREE_PLATFORM_SHARED_LIB_EXTENSION;
+	FString SharedLibDirPath = UE::NNERuntimeIREECpu::Private::GetPackagedModelDirPath(NNE_RUNTIME_IREE_PLATFORM_NAME);
+	FString SharedLibName = FileIdString + "." + NNE_RUNTIME_IREE_PLATFORM_SHARED_LIB_EXTENSION;
 #endif // WITH_EDITOR
 
 #if WITH_EDITOR
diff --git a/Engine/Plugins/Experimental/NNERuntimeIREE/Source/NNERuntimeIREE/Private/NNERuntimeIREEModule.cpp b/Engine/Plugins/Experimental/NNERuntimeIREE/Source/NNERuntimeIREE/Private/NNERuntimeIREEModule.cpp
--- a/Engine/Plugins/Experimental/NNERuntimeIREE/Source/NNERuntimeIREE/Private/NNERuntimeIREEModule.cpp
+++ b/Engine/Plugins/Experimental/NNERuntimeIREE/Source/NNERuntimeIREE/Private/NNERuntimeIREEModule.cpp
@@ -3,13 +3,62 @@
 #include "NNERuntimeIREEModule.h"
 
 #ifdef WITH_NNE_RUNTIME_IREE
+
 #include "NNE.h"
 #include "NNERuntime.h"
+
+#ifdef WITH_EDITOR
+#include "Interfaces/ITargetPlatformManagerModule.h"
+#include "Interfaces/ITargetPlatform.h"
+#include "Misc/App.h"
+#include "Misc/ConfigCacheIni.h"
+#endif // WITH_EDITOR
+
 #endif // WITH_NNE_RUNTIME_IREE
 
 void FNNERuntimeIREEModule::StartupModule()
 {
 #ifdef WITH_NNE_RUNTIME_IREE
+
+#ifdef WITH_EDITOR
+	ITargetPlatformManagerModule* TargetPlatformManagerModule = GetTargetPlatformManager();
+	if (TargetPlatformManagerModule)
+	{
+		TSet<FString> ProcessedPlatforms;
+		TArray<ITargetPlatform*> TargetPlatforms = TargetPlatformManagerModule->GetTargetPlatforms();
+		for (int32 i = 0; i < TargetPlatforms.Num(); i++)
+		{
+			FString IniPlatformName = TargetPlatforms[i]->IniPlatformName();
+			if (!ProcessedPlatforms.Contains(IniPlatformName))
+			{
+				FString TargetPlatformDisplayName = IniPlatformName;
+				FString TargetPlatformName = IniPlatformName.Equals("Windows") ? "Win64" : IniPlatformName;
+
+				FString ConfigFolderPath = FPaths::ConvertRelativePathToFull(FPaths::ProjectConfigDir());
+				FString ConfigFilePath = FPaths::Combine(ConfigFolderPath, TargetPlatformDisplayName, TargetPlatformDisplayName + "Game.ini");
+				FConfigFile OnDiskConfigFile;
+				OnDiskConfigFile.Read(ConfigFilePath);
+
+				FString CookingPath = FString("/Saved/Cooked/") + TargetPlatformDisplayName + FString("/Engine/Plugins/") + UE_PLUGIN_NAME + FString("/Binaries");
+				FString PackagingPath = FString("/Binaries/") + TargetPlatformName + FString("/") + UE_PLUGIN_NAME;
+
+				OnDiskConfigFile.AddUniqueToSection(TEXT("/Script/UnrealEd.ProjectPackagingSettings"), TEXT("+DirectoriesToAlwaysStageAsNonUFS"), FString("(Path=\"..") + CookingPath + FString("\")"));
+				OnDiskConfigFile.AddUniqueToSection(TEXT("Staging"), TEXT("+RemapDirectories"), FString("(From=\"") + FApp::GetProjectName() + CookingPath + FString("\", To=\"") + FApp::GetProjectName() + PackagingPath + FString("\")"));
+				OnDiskConfigFile.AddUniqueToSection(TEXT("Staging"), TEXT("+AllowedDirectories"), FApp::GetProjectName() + PackagingPath);
+
+				// Write returns true without an attempt to write if the config has not been marked dirty and AddUniqueToSection will only mark it dirty if there are changes
+				// Thus the warning will not be printed when a locked file contains the required settings
+				if (!OnDiskConfigFile.Write(ConfigFilePath))
+				{
+					UE_LOG(LogTemp, Warning, TEXT("UNNERuntimeIREECpu could not write config file %s. Please make the file writeable and re-cook or manually add the required staging settings or models will not work in packaged builds for platform %s!"), *ConfigFilePath, *TargetPlatformDisplayName);
+				}
+
+				ProcessedPlatforms.Add(IniPlatformName);
+			}
+		}
+	}
+#endif // WITH_EDITOR
+
 	NNERuntimeIREECpu = NewObject<UNNERuntimeIREECpu>();
 	if (NNERuntimeIREECpu.IsValid())
 	{
