[NNE] Cleanup NNERuntimeIREE paths and move warning from module startup to cooking where it is actually useful.

Problem: NNERuntimeIREE would complain about non-writeable config files inside the startup even when it does not matter.

Solution: Move the check and warning do CreateModelData where it actually needs to be triggered.

Tests: Tested different scenarios of creating config files (including non-writeable files and files containing other settings), built and run on win64, non-unity compilation and presubmit

#rnx

diff --git a/Engine/Plugins/Experimental/NNERuntimeIREE/Source/NNERuntimeIREE/Private/NNERuntimeIREECpu.cpp b/Engine/Plugins/Experimental/NNERuntimeIREE/Source/NNERuntimeIREE/Private/NNERuntimeIREECpu.cpp
--- a/Engine/Plugins/Experimental/NNERuntimeIREE/Source/NNERuntimeIREE/Private/NNERuntimeIREECpu.cpp
+++ b/Engine/Plugins/Experimental/NNERuntimeIREE/Source/NNERuntimeIREE/Private/NNERuntimeIREECpu.cpp
@@ -27,6 +27,7 @@
 #include "Internationalization/TextLocalizationResource.h"
 #include "Memory/SharedBuffer.h"
 #include "Misc/App.h"
+#include "Misc/ConfigCacheIni.h"
 #include "Serialization/JsonReader.h"
 #include "Serialization/JsonSerializer.h"
 #include "Serialization/JsonSerializerMacros.h"
@@ -217,24 +218,6 @@ namespace UE::NNERuntimeIREECpu::Private
 	}
 #endif // WITH_EDITOR
 
-	FString GetIntermediateModelDirPath(const FString& PlatformName, const FString& FileIdString)
-	{
-		FString IntermediateDir = FPaths::ConvertRelativePathToFull(*FPaths::ProjectIntermediateDir());
-		return FPaths::Combine(IntermediateDir, "Build", PlatformName, UE_PLUGIN_NAME, FileIdString);
-	}
-
-	FString GetCookedModelDirPath(const FString& PlatformDisplayName)
-	{
-		FString SavedDir = FPaths::ConvertRelativePathToFull(*FPaths::ProjectSavedDir());
-		return FPaths::Combine(SavedDir, "Cooked", PlatformDisplayName, "Engine", "Plugins", UE_PLUGIN_NAME, "Binaries");
-	}
-
-	FString GetPackagedModelDirPath(const FString& PlatformName)
-	{
-		FString ProjectDir = FPaths::ConvertRelativePathToFull(*FPaths::ProjectDir());
-		return FPaths::Combine(ProjectDir, "Binaries", PlatformName, UE_PLUGIN_NAME);
-	}
-
 	FString GetModelCpuDataIdentifier(const FString& FileIdString, const FString& PlatformDisplayName)
 	{
 		return FileIdString + "-" + UNNERuntimeIREECpu::GUID.ToString(EGuidFormats::Digits) + "-" + FString::FromInt(UNNERuntimeIREECpu::Version) + "-" + FString::FromInt(UNNERuntimeIREECpu::MemoryAlignment) + "-" + PlatformDisplayName;
@@ -294,6 +277,14 @@ TSharedPtr<UE::NNE::FSharedModelData> UNNERuntimeIREECpu::CreateModelData(FStrin
 		return TSharedPtr<UE::NNE::FSharedModelData>();
 	}
 
+	FConfigFile ConfigFile;
+	FString ConfigFilePath;
+	GetUpdatedPlatformConfig(TargetPlatformDisplayName, ConfigFile, ConfigFilePath);
+	if (ConfigFile.Dirty)
+	{
+		UE_LOG(LogTemp, Warning, TEXT("UNNERuntimeIREECpu could not find the required settings in config file %s. Please make the file writeable and re-start the editor or manually add the required staging settings or models will not work in packaged builds for platform %s!"), *ConfigFilePath, *TargetPlatformDisplayName);
+	}
+
 	FString BuildConfigFileName = FString("IREE_") + NNE_RUNTIME_IREE_PLATFORM_NAME + "_To_" + TargetPlatformDisplayName + ".json";
 	UE::NNERuntimeIREECpu::Private::FBuildConfig BuildConfig;
 	if (!UE::NNERuntimeIREECpu::Private::LoadBuildConfig(TargetPlatformDisplayName, BuildConfigFileName, BuildConfig))
@@ -304,7 +295,7 @@ TSharedPtr<UE::NNE::FSharedModelData> UNNERuntimeIREECpu::CreateModelData(FStrin
 
 	FString FileIdString = FileId.ToString(EGuidFormats::Digits).ToLower();
 	FString PluginDir = FPaths::ConvertRelativePathToFull(*IPluginManager::Get().FindPlugin(UE_PLUGIN_NAME)->GetBaseDir());
-	FString IntermediateModelDir = UE::NNERuntimeIREECpu::Private::GetIntermediateModelDirPath(BuildConfig.TargetPlatformName, FileIdString);
+	FString IntermediateModelDir = FPaths::ConvertRelativePathToFull(FPaths::Combine(FPaths::ProjectDir(), GetIntermediateModelDirPath(BuildConfig.TargetPlatformName, FileIdString)));
 	FString ObjectPath = FPaths::Combine(IntermediateModelDir, FileIdString + ".o");
 	FString AdditionalOutputPath = FPaths::Combine(IntermediateModelDir, FileIdString + ".vmfb");
 	FString InputPath = FPaths::Combine(IntermediateModelDir, FileIdString + "." + FileType);
@@ -354,7 +345,7 @@ TSharedPtr<UE::NNE::FSharedModelData> UNNERuntimeIREECpu::CreateModelData(FStrin
 			return TSharedPtr<UE::NNE::FSharedModelData>();
 		}
 
-		FString CookedModelPath = FPaths::Combine(UE::NNERuntimeIREECpu::Private::GetCookedModelDirPath(TargetPlatformDisplayName), SharedLibName);
+		FString CookedModelPath = FPaths::ConvertRelativePathToFull(FPaths::Combine(FPaths::ProjectDir(), GetCookedModelDirPath(TargetPlatformDisplayName), SharedLibName));
 		IFileManager::Get().Copy(*CookedModelPath, *SharedLibPath);
 
 		TArray<uint8> SharedLibData;
@@ -495,10 +486,10 @@ TSharedPtr<UE::NNE::IModelCPU> UNNERuntimeIREECpu::CreateModelCPU(TObjectPtr<UNN
 
 	FString FileIdString = FileId.ToString(EGuidFormats::Digits).ToLower();
 #if WITH_EDITOR
-	FString SharedLibDirPath = UE::NNERuntimeIREECpu::Private::GetIntermediateModelDirPath(NNE_RUNTIME_IREE_PLATFORM_NAME, FileIdString);
+	FString SharedLibDirPath = FPaths::ConvertRelativePathToFull(FPaths::Combine(FPaths::ProjectDir(), GetIntermediateModelDirPath(NNE_RUNTIME_IREE_PLATFORM_NAME, FileIdString)));
 	FString SharedLibName = FileIdString + "." + NNE_RUNTIME_IREE_PLATFORM_SHARED_LIB_EXTENSION;
 #else
-	FString SharedLibDirPath = UE::NNERuntimeIREECpu::Private::GetPackagedModelDirPath(NNE_RUNTIME_IREE_PLATFORM_NAME);
+	FString SharedLibDirPath = FPaths::ConvertRelativePathToFull(FPaths::Combine(FPaths::ProjectDir(), GetPackagedModelDirPath(NNE_RUNTIME_IREE_PLATFORM_NAME)));
 	FString SharedLibName = FileIdString + "." + NNE_RUNTIME_IREE_PLATFORM_SHARED_LIB_EXTENSION;
 #endif // WITH_EDITOR
 
@@ -538,6 +529,37 @@ TSharedPtr<UE::NNE::IModelCPU> UNNERuntimeIREECpu::CreateModelCPU(TObjectPtr<UNN
 	return TSharedPtr<UE::NNE::IModelCPU>(static_cast<UE::NNE::IModelCPU*>(Model.Release()));
 }
 
+FString UNNERuntimeIREECpu::GetIntermediateModelDirPath(const FString& PlatformName, const FString& FileIdString)
+{
+	return FPaths::Combine("Intermediate", "Build", PlatformName, UE_PLUGIN_NAME, FileIdString);
+}
+
+FString UNNERuntimeIREECpu::GetCookedModelDirPath(const FString& PlatformName)
+{
+	return FPaths::Combine("Saved", "Cooked", PlatformName, "Engine", "Plugins", UE_PLUGIN_NAME, "Binaries");
+}
+
+FString UNNERuntimeIREECpu::GetPackagedModelDirPath(const FString& PlatformName)
+{
+	FString PlatformNameShort = PlatformName.Equals("Windows") ? "Win64" : PlatformName;
+	return FPaths::Combine("Binaries", PlatformNameShort, UE_PLUGIN_NAME);
+}
+
+void UNNERuntimeIREECpu::GetUpdatedPlatformConfig(const FString& PlatformName, FConfigFile& ConfigFile, FString& ConfigFilePath)
+{ 
+	FString ConfigFolderPath = FPaths::ConvertRelativePathToFull(FPaths::ProjectConfigDir());
+	ConfigFilePath = FPaths::Combine(ConfigFolderPath, PlatformName, PlatformName + "Game.ini");
+
+	ConfigFile.Read(ConfigFilePath);
+
+	FString CookingPath = FString("/") + GetCookedModelDirPath(PlatformName);
+	FString PackagingPath = FString("/") + GetPackagedModelDirPath(PlatformName);
+
+	ConfigFile.AddUniqueToSection(TEXT("/Script/UnrealEd.ProjectPackagingSettings"), TEXT("+DirectoriesToAlwaysStageAsNonUFS"), FString("(Path=\"..") + CookingPath + FString("\")"));
+	ConfigFile.AddUniqueToSection(TEXT("Staging"), TEXT("+RemapDirectories"), FString("(From=\"") + FApp::GetProjectName() + CookingPath + FString("\", To=\"") + FApp::GetProjectName() + PackagingPath + FString("\")"));
+	ConfigFile.AddUniqueToSection(TEXT("Staging"), TEXT("+AllowedDirectories"), FApp::GetProjectName() + PackagingPath);
+}
+
 #else // WITH_NNE_RUNTIME_IREE
 
 FString UNNERuntimeIREECpu::GetRuntimeName() const { return ""; };
@@ -549,4 +571,10 @@ FString UNNERuntimeIREECpu::GetModelDataIdentifier(FString FileType, TConstArray
 bool UNNERuntimeIREECpu::CanCreateModelCPU(TObjectPtr<UNNEModelData> ModelData) const { return false; };
 TSharedPtr<UE::NNE::IModelCPU> UNNERuntimeIREECpu::CreateModelCPU(TObjectPtr<UNNEModelData> ModelData) { return TSharedPtr<UE::NNE::IModelCPU>(); };
 
+FString UNNERuntimeIREECpu::GetIntermediateModelDirPath(const FString& PlatformName, const FString& FileIdString) { return ""; }
+FString UNNERuntimeIREECpu::GetCookedModelDirPath(const FString& PlatformName) { return ""; }
+FString UNNERuntimeIREECpu::GetPackagedModelDirPath(const FString& PlatformName) { return ""; }
+
+void UNNERuntimeIREECpu::GetUpdatedPlatformConfig(const FString& PlatformName, FConfigFile& ConfigFile, FString& ConfigFilePath) { }
+
 #endif // WITH_NNE_RUNTIME_IREE
\ No newline at end of file
diff --git a/Engine/Plugins/Experimental/NNERuntimeIREE/Source/NNERuntimeIREE/Private/NNERuntimeIREECpu.h b/Engine/Plugins/Experimental/NNERuntimeIREE/Source/NNERuntimeIREE/Private/NNERuntimeIREECpu.h
--- a/Engine/Plugins/Experimental/NNERuntimeIREE/Source/NNERuntimeIREE/Private/NNERuntimeIREECpu.h
+++ b/Engine/Plugins/Experimental/NNERuntimeIREE/Source/NNERuntimeIREE/Private/NNERuntimeIREECpu.h
@@ -3,6 +3,7 @@
 #pragma once
 
 #include "CoreMinimal.h"
+#include "Misc/ConfigCacheIni.h"
 #include "Misc/Guid.h"
 #include "NNEModelData.h"
 #include "NNERuntime.h"
@@ -34,4 +35,10 @@ public:
 	virtual bool CanCreateModelCPU(TObjectPtr<UNNEModelData> ModelData) const override;
 	virtual TSharedPtr<UE::NNE::IModelCPU> CreateModelCPU(TObjectPtr<UNNEModelData> ModelData) override;
 	//~ End INNERuntimeCPU Interface
+
+	static FString GetIntermediateModelDirPath(const FString& PlatformName, const FString& FileIdString);
+	static FString GetCookedModelDirPath(const FString& PlatformName);
+	static FString GetPackagedModelDirPath(const FString& PlatformName);
+
+	static void GetUpdatedPlatformConfig(const FString& PlatformName, FConfigFile& ConfigFile, FString& ConfigFilePath);
 };
\ No newline at end of file
diff --git a/Engine/Plugins/Experimental/NNERuntimeIREE/Source/NNERuntimeIREE/Private/NNERuntimeIREEModule.cpp b/Engine/Plugins/Experimental/NNERuntimeIREE/Source/NNERuntimeIREE/Private/NNERuntimeIREEModule.cpp
--- a/Engine/Plugins/Experimental/NNERuntimeIREE/Source/NNERuntimeIREE/Private/NNERuntimeIREEModule.cpp
+++ b/Engine/Plugins/Experimental/NNERuntimeIREE/Source/NNERuntimeIREE/Private/NNERuntimeIREEModule.cpp
@@ -8,9 +8,10 @@
 #include "NNERuntime.h"
 
 #ifdef WITH_EDITOR
+#include "Serialization/Archive.h"
+#include "HAL/FileManager.h"
 #include "Interfaces/ITargetPlatformManagerModule.h"
 #include "Interfaces/ITargetPlatform.h"
-#include "Misc/App.h"
 #include "Misc/ConfigCacheIni.h"
 #endif // WITH_EDITOR
 
@@ -19,6 +20,12 @@
 void FNNERuntimeIREEModule::StartupModule()
 {
 #ifdef WITH_NNE_RUNTIME_IREE
+	NNERuntimeIREECpu = NewObject<UNNERuntimeIREECpu>();
+	if (NNERuntimeIREECpu.IsValid())
+	{
+		NNERuntimeIREECpu->AddToRoot();
+		UE::NNE::RegisterRuntime(NNERuntimeIREECpu.Get());
+	}
 
 #ifdef WITH_EDITOR
 	ITargetPlatformManagerModule* TargetPlatformManagerModule = GetTargetPlatformManager();
@@ -31,41 +38,27 @@ void FNNERuntimeIREEModule::StartupModule()
 			FString IniPlatformName = TargetPlatforms[i]->IniPlatformName();
 			if (!ProcessedPlatforms.Contains(IniPlatformName))
 			{
+				ProcessedPlatforms.Add(IniPlatformName);
 				FString TargetPlatformDisplayName = IniPlatformName;
-				FString TargetPlatformName = IniPlatformName.Equals("Windows") ? "Win64" : IniPlatformName;
-
-				FString ConfigFolderPath = FPaths::ConvertRelativePathToFull(FPaths::ProjectConfigDir());
-				FString ConfigFilePath = FPaths::Combine(ConfigFolderPath, TargetPlatformDisplayName, TargetPlatformDisplayName + "Game.ini");
-				FConfigFile OnDiskConfigFile;
-				OnDiskConfigFile.Read(ConfigFilePath);
-
-				FString CookingPath = FString("/Saved/Cooked/") + TargetPlatformDisplayName + FString("/Engine/Plugins/") + UE_PLUGIN_NAME + FString("/Binaries");
-				FString PackagingPath = FString("/Binaries/") + TargetPlatformName + FString("/") + UE_PLUGIN_NAME;
-
-				OnDiskConfigFile.AddUniqueToSection(TEXT("/Script/UnrealEd.ProjectPackagingSettings"), TEXT("+DirectoriesToAlwaysStageAsNonUFS"), FString("(Path=\"..") + CookingPath + FString("\")"));
-				OnDiskConfigFile.AddUniqueToSection(TEXT("Staging"), TEXT("+RemapDirectories"), FString("(From=\"") + FApp::GetProjectName() + CookingPath + FString("\", To=\"") + FApp::GetProjectName() + PackagingPath + FString("\")"));
-				OnDiskConfigFile.AddUniqueToSection(TEXT("Staging"), TEXT("+AllowedDirectories"), FApp::GetProjectName() + PackagingPath);
-
-				// Write returns true without an attempt to write if the config has not been marked dirty and AddUniqueToSection will only mark it dirty if there are changes
-				// Thus the warning will not be printed when a locked file contains the required settings
-				if (!OnDiskConfigFile.Write(ConfigFilePath))
+				FConfigFile ConfigFile;
+				FString ConfigFilePath;
+				NNERuntimeIREECpu->GetUpdatedPlatformConfig(IniPlatformName, ConfigFile, ConfigFilePath);
+				if (ConfigFile.Dirty)
 				{
-					UE_LOG(LogTemp, Warning, TEXT("UNNERuntimeIREECpu could not write config file %s. Please make the file writeable and re-cook or manually add the required staging settings or models will not work in packaged builds for platform %s!"), *ConfigFilePath, *TargetPlatformDisplayName);
+					{
+						TUniquePtr<FArchive> Ar = TUniquePtr<FArchive>(IFileManager::Get().CreateFileWriter(*ConfigFilePath, EFileWrite::FILEWRITE_Append));
+						if (!Ar)
+						{
+							continue;
+						}
+					}
+					ConfigFile.Write(ConfigFilePath);
 				}
-
-				ProcessedPlatforms.Add(IniPlatformName);
 			}
 		}
 	}
 #endif // WITH_EDITOR
 
-	NNERuntimeIREECpu = NewObject<UNNERuntimeIREECpu>();
-	if (NNERuntimeIREECpu.IsValid())
-	{
-		NNERuntimeIREECpu->AddToRoot();
-		UE::NNE::RegisterRuntime(NNERuntimeIREECpu.Get());
-	}
-
 	NNERuntimeIREECuda = NewObject<UNNERuntimeIREECuda>();
 	if (NNERuntimeIREECuda.IsValid())
 	{
