[NNE] Improving the NNERuntimeIREECpu build system to allow building multiple architectures for the same target platform

Problem: Certain platforms have multiple architectures, each one with different instruction sets. Currently our IREE runtime does not make use of them.

Solution: Add target architecture configurations to our building pipeline and decide at runtime which one to load.

Tested: Win, Linux, consoles, non-unity compilation and presubmit

#rnx

diff --git a/Engine/Plugins/Experimental/NNERuntimeIREE/Config/IREE_Linux_To_Linux.json b/Engine/Plugins/Experimental/NNERuntimeIREE/Config/IREE_Linux_To_Linux.json
--- a/Engine/Plugins/Experimental/NNERuntimeIREE/Config/IREE_Linux_To_Linux.json
+++ b/Engine/Plugins/Experimental/NNERuntimeIREE/Config/IREE_Linux_To_Linux.json
@@ -1,8 +1,12 @@
 {
-	"TargetPlatformName": "Linux",
-	"SharedLibraryExtension": "so",
-	"ModelCompilationCommand": "${PLUGIN_DIR}/Source/ThirdParty/IREE/Bin/Linux/iree-compile",
-	"ModelCompilationArguments": "--iree-input-type=auto --iree-llvmcpu-target-triple=x86_64-unknown-linux-gnu --iree-hal-target-backends=llvm-cpu --iree-llvmcpu-link-embedded=false --iree-llvmcpu-link-static --iree-llvmcpu-static-library-output-path=\"${OBJECT_PATH}\" --output-format=vm-bytecode -o \"${ADDITIONAL_OUTPUT_PATH}\" \"${INPUT_PATH}\"",
-	"SharedLibraryCompilationCommand": "${LINUX_COMPILER}",
-	"SharedLibraryCompilationArguments": "-shared -Wl,--export-dynamic -o \"${SHARED_LIB_PATH}\" \"${OBJECT_PATH}\""
+	"CompilerCommand": [ "${PLUGIN_DIR}/Source/ThirdParty/IREE/Bin/Linux/iree-compile" ],
+	"LinkerCommand": [ "${PLUGIN_DIR}/Source/ThirdParty/IREE/Bin/Linux/ld.lld" ],
+	"SharedLibExt": ".so",
+	"Targets": [
+		{
+			"Architecture": "",
+			"CompilerArguments": "--iree-llvmcpu-target-triple=x86_64-unknown-linux-gnu --iree-input-type=auto --iree-hal-target-backends=llvm-cpu --iree-llvmcpu-link-embedded=false --iree-llvmcpu-link-static --output-format=vm-bytecode --iree-llvmcpu-static-library-output-path=${OBJECT_PATH} -o ${VMFB_PATH} ${INPUT_PATH}",
+			"LinkerArguments": "-shared --export-dynamic -o ${SHARED_LIB_PATH} ${OBJECT_PATH}"
+		}
+	]
 }
diff --git a/Engine/Plugins/Experimental/NNERuntimeIREE/Config/IREE_Mac_To_Mac.json b/Engine/Plugins/Experimental/NNERuntimeIREE/Config/IREE_Mac_To_Mac.json
--- a/Engine/Plugins/Experimental/NNERuntimeIREE/Config/IREE_Mac_To_Mac.json
+++ b/Engine/Plugins/Experimental/NNERuntimeIREE/Config/IREE_Mac_To_Mac.json
@@ -1,8 +1,17 @@
 {
-	"TargetPlatformName": "Mac",
-	"SharedLibraryExtension": "dylib",
-	"ModelCompilationCommand": "${PLUGIN_DIR}/Source/ThirdParty/IREE/Bin/Mac/iree-compile",
-	"ModelCompilationArguments": "--iree-input-type=auto --iree-llvmcpu-target-triple=x86_64-apple-darwin --iree-hal-target-backends=llvm-cpu --iree-llvmcpu-link-embedded=false --iree-llvmcpu-link-static --iree-llvmcpu-static-library-output-path=\"${OBJECT_PATH}\" --output-format=vm-bytecode -o \"${ADDITIONAL_OUTPUT_PATH}\" \"${INPUT_PATH}\"",
-	"SharedLibraryCompilationCommand": "/usr/bin/clang++",
-	"SharedLibraryCompilationArguments": "-dynamiclib -o \"${SHARED_LIB_PATH}\" \"${OBJECT_PATH}\""
+	"CompilerCommand": [ "${PLUGIN_DIR}/Source/ThirdParty/IREE/Bin/Mac/iree-compile" ],
+	"LinkerCommand": [ "${PLUGIN_DIR}/Source/ThirdParty/IREE/Bin/Mac/ld64.lld" ],
+	"SharedLibExt": ".dylib",
+	"Targets": [
+		{
+			"Architecture": "x86_64",
+			"CompilerArguments": "--iree-llvmcpu-target-triple=x86_64-apple-darwin --iree-input-type=auto --iree-hal-target-backends=llvm-cpu --iree-llvmcpu-link-embedded=false --iree-llvmcpu-link-static --output-format=vm-bytecode --iree-llvmcpu-static-library-output-path=${OBJECT_PATH} -o ${VMFB_PATH} ${INPUT_PATH}",
+			"LinkerArguments": "-dylib -arch x86_64 ${SHARED_LIB_PATH} ${OBJECT_PATH}"
+		},
+		{
+			"Architecture": "arm64",
+			"CompilerArguments": "--iree-llvmcpu-target-triple=arm64-apple-darwin --iree-input-type=auto --iree-hal-target-backends=llvm-cpu --iree-llvmcpu-link-embedded=false --iree-llvmcpu-link-static --output-format=vm-bytecode --iree-llvmcpu-static-library-output-path=${OBJECT_PATH} -o ${VMFB_PATH} ${INPUT_PATH}",
+			"LinkerArguments": "-dylib -arch arm64 ${SHARED_LIB_PATH} ${OBJECT_PATH}"
+		}
+	]
 }
diff --git a/Engine/Plugins/Experimental/NNERuntimeIREE/Config/IREE_Win64_To_Windows.json b/Engine/Plugins/Experimental/NNERuntimeIREE/Config/IREE_Win64_To_Windows.json
deleted file mode 100644
diff --git a/Engine/Plugins/Experimental/NNERuntimeIREE/Config/IREE_Windows_To_Windows.json b/Engine/Plugins/Experimental/NNERuntimeIREE/Config/IREE_Windows_To_Windows.json
new file mode 100644
--- /dev/null
+++ b/Engine/Plugins/Experimental/NNERuntimeIREE/Config/IREE_Windows_To_Windows.json
@@ -0,0 +1,12 @@
+{
+	"CompilerCommand": [ "${PLUGIN_DIR}/Source/ThirdParty/IREE/Bin/Windows/iree-compile.exe" ],
+	"LinkerCommand": [ "${PLUGIN_DIR}/Source/ThirdParty/IREE/Bin/Windows/lld-link.exe" ],
+	"SharedLibExt": ".dll",
+	"Targets": [
+		{
+			"Architecture": "",
+			"CompilerArguments": "--iree-llvmcpu-target-triple=x86_64-pc-windows-msvc --iree-input-type=auto --iree-hal-target-backends=llvm-cpu --iree-llvmcpu-link-embedded=false --iree-llvmcpu-link-static --output-format=vm-bytecode --iree-llvmcpu-static-library-output-path=${OBJECT_PATH} -o ${VMFB_PATH} ${INPUT_PATH}",
+			"LinkerArguments": "/DLL /NOENTRY /FORCE:UNRESOLVED /OUT:${SHARED_LIB_PATH} ${OBJECT_PATH}"
+		}
+	]
+}
\ No newline at end of file
diff --git a/Engine/Plugins/Experimental/NNERuntimeIREE/Source/NNERuntimeIREE/NNERuntimeIREE.Build.cs b/Engine/Plugins/Experimental/NNERuntimeIREE/Source/NNERuntimeIREE/NNERuntimeIREE.Build.cs
--- a/Engine/Plugins/Experimental/NNERuntimeIREE/Source/NNERuntimeIREE/NNERuntimeIREE.Build.cs
+++ b/Engine/Plugins/Experimental/NNERuntimeIREE/Source/NNERuntimeIREE/NNERuntimeIREE.Build.cs
@@ -1,23 +1,9 @@
 // Copyright Epic Games, Inc. All Rights Reserved.
 
-using System.IO;
 using UnrealBuildTool;
-using EpicGames.Core;
-using System;
 
 public class NNERuntimeIREE : ModuleRules
 {
-	protected static readonly string SharedLibName = "NNERuntimeIREEModels";
-
-	protected void SetupNNERuntimeIREE(string PlatformName, string PlatformDisplayName, string SharedLibExtension)
-	{
-		PublicDefinitions.Add("WITH_NNE_RUNTIME_IREE");
-
-		PrivateDefinitions.Add("NNE_RUNTIME_IREE_PLATFORM_NAME=\"" + PlatformName + "\"");
-		PrivateDefinitions.Add("NNE_RUNTIME_IREE_PLATFORM_DISPLAY_NAME=\"" + PlatformDisplayName + "\"");
-		PrivateDefinitions.Add("NNE_RUNTIME_IREE_PLATFORM_SHARED_LIB_EXTENSION=\"" + SharedLibExtension + "\"");
-	}
-
 	public NNERuntimeIREE(ReadOnlyTargetRules Target) : base(Target)
 	{
 		PCHUsage = PCHUsageMode.UseExplicitOrSharedPCHs;
@@ -30,7 +16,8 @@ public class NNERuntimeIREE : ModuleRules
 				"CoreUObject",
 				"Engine",
 				"IREE",
-				"NNE"
+				"NNE",
+				"Projects"
 			}
 		);
 
@@ -41,7 +28,6 @@ public class NNERuntimeIREE : ModuleRules
 				{
 					"DerivedDataCache",
 					"Json",
-					"Projects",
 					"TargetPlatform"
 				}
 			);
@@ -49,22 +35,16 @@ public class NNERuntimeIREE : ModuleRules
 
 		if (Target.Platform == UnrealTargetPlatform.Win64)
 		{
-			SetupNNERuntimeIREE("Win64", "Windows", "dll");
-
-			PrivateDefinitions.Add("NNE_RUNTIME_IREE_WIN_TOOLCHAIN_PATH=\"" + Target.WindowsPlatform.ToolChainDir.Replace("\\", "/") + "/\"");
-			PrivateDefinitions.Add("NNE_RUNTIME_IREE_WIN_SDK_LIB_PATH=\"" + Target.WindowsPlatform.WindowsSdkDir.Replace("\\", "/") + "/Lib/" + Target.WindowsPlatform.WindowsSdkVersion + "/\"");
+			PublicDefinitions.Add("WITH_NNE_RUNTIME_IREE");
 		}
 		else if (Target.Platform == UnrealTargetPlatform.Linux)
 		{
-			SetupNNERuntimeIREE("Linux", "Linux", "so");
-
-			PrivateDefinitions.Add("NNE_RUNTIME_IREE_LINUX_COMPILER=\"" + UEBuildPlatformSDK.GetSDKForPlatform("Linux").GetInternalSDKPath().Replace("\\", "/") + "/bin/clang++\"");
+			PublicDefinitions.Add("WITH_NNE_RUNTIME_IREE");
 			PrivateDefinitions.Add("NNE_RUNTIME_IREE_USE_COMBINED_LIB_PATH");
 		}
 		else if (Target.Platform == UnrealTargetPlatform.Mac)
 		{
-			SetupNNERuntimeIREE("Mac", "Mac", "dylib");
-
+			PublicDefinitions.Add("WITH_NNE_RUNTIME_IREE");
 			PrivateDefinitions.Add("NNE_RUNTIME_IREE_USE_COMBINED_LIB_PATH");
 		}
 	}
diff --git a/Engine/Plugins/Experimental/NNERuntimeIREE/Source/NNERuntimeIREE/Private/NNERuntimeIREECommon.cpp b/Engine/Plugins/Experimental/NNERuntimeIREE/Source/NNERuntimeIREE/Private/NNERuntimeIREECommon.cpp
--- a/Engine/Plugins/Experimental/NNERuntimeIREE/Source/NNERuntimeIREE/Private/NNERuntimeIREECommon.cpp
+++ b/Engine/Plugins/Experimental/NNERuntimeIREE/Source/NNERuntimeIREE/Private/NNERuntimeIREECommon.cpp
@@ -149,11 +149,6 @@ namespace UE::NNERuntimeIREE
 		}
 	} // Private
 
-	FString GetTargetPlatformDisplayName(const ITargetPlatform* TargetPlatform)
-	{
-		return TargetPlatform ? TargetPlatform->IniPlatformName() : UGameplayStatics::GetPlatformName();
-	}
-
 	class FIREEInstance
 	{
 	public:
@@ -170,6 +165,8 @@ namespace UE::NNERuntimeIREE
 		TSharedPtr<FIREELibrary> GetLibrary(const FString& Key);
 		void SetLibrary(const FString& Key, TSharedPtr<FIREELibrary> Library);
 
+		TSharedPtr<UE::NNE::FSharedModelData> GetModelData(const FString& DirPath, const FString& VmfbFileName);
+
 	public:
 		static TSharedPtr<FIREEInstance> GetInstance();
 
@@ -178,6 +175,7 @@ namespace UE::NNERuntimeIREE
 		static FCriticalSection CriticalSection;
 		iree_runtime_instance_t* Instance;
 		TMap<FString, TWeakPtr<FIREELibrary>> Libraries;
+		TMap<FString, TWeakPtr<UE::NNE::FSharedModelData>> ModelData;
 	};
 
 	TWeakPtr<FIREEInstance> FIREEInstance::WeakInstancePtr;
@@ -207,7 +205,7 @@ namespace UE::NNERuntimeIREE
 
 	bool FIREEInstance::CreateModule(TConstArrayView<uint8> VmfbDataView, iree_vm_module_t** Module)
 	{
-		check(VmfbDataView.Num() > 0);
+		check(!VmfbDataView.IsEmpty());
 		check(Module);
 
 		iree_status_t Status = iree_ok_status();
@@ -346,7 +344,7 @@ namespace UE::NNERuntimeIREE
 
 	TSharedPtr<FIREELibrary> FIREEInstance::GetLibrary(const FString& Key)
 	{
-		check(Key.Len() > 0);
+		check(!Key.IsEmpty());
 		if (Libraries.Contains(Key))
 		{
 			return Libraries[Key].Pin();
@@ -356,12 +354,44 @@ namespace UE::NNERuntimeIREE
 
 	void FIREEInstance::SetLibrary(const FString& Key, TSharedPtr<FIREELibrary> Library)
 	{
-		check(Key.Len() > 0);
+		check(!Key.IsEmpty());
 		check(Library.IsValid());
 		check(!Libraries.Contains(Key) || !Libraries[Key].IsValid());
 		Libraries.Emplace(Key, Library);
 	}
 
+	TSharedPtr<UE::NNE::FSharedModelData> FIREEInstance::GetModelData(const FString& DirPath, const FString& VmfbFileName)
+	{
+		check(!VmfbFileName.IsEmpty());
+
+		FString FilePath = FPaths::Combine(DirPath, VmfbFileName);
+		if (ModelData.Contains(FilePath) && ModelData[FilePath].IsValid())
+		{
+			return ModelData[FilePath].Pin();
+		}
+
+		TUniquePtr<FArchive> Reader = TUniquePtr<FArchive>(IFileManager::Get().CreateFileReader(*FilePath, 0));
+		if (!Reader)
+		{
+			UE_LOG(LogTemp, Error, TEXT("Failed to open the vmfb data file '%s'"), *FilePath);
+			return TSharedPtr<UE::NNE::FSharedModelData>();
+		}
+		int64 DataSize = Reader->TotalSize();
+		if (DataSize < 1)
+		{
+			UE_LOG(LogTemp, Error, TEXT("Vmfb data file '%s' is empty"), *FilePath);
+			return TSharedPtr<UE::NNE::FSharedModelData>();
+		}
+
+		void* Data = FMemory::Malloc(DataSize, IREE_HAL_HEAP_BUFFER_ALIGNMENT);
+		Reader->Serialize(Data, DataSize);
+
+		TSharedPtr<UE::NNE::FSharedModelData> Result = MakeShared<UE::NNE::FSharedModelData>(FSharedBuffer::TakeOwnership(Data, DataSize, FMemory::Free), IREE_HAL_HEAP_BUFFER_ALIGNMENT);
+		ModelData.Emplace(FilePath, Result);
+
+		return Result;
+	}
+
 	TSharedPtr<FIREEInstance> FIREEInstance::GetInstance()
 	{
 		FScopeLock ScopeLock(&FIREEInstance::CriticalSection);
@@ -419,10 +449,9 @@ namespace UE::NNERuntimeIREE
 		Instance.Reset();
 	}
 
-	TSharedPtr<FIREEModule> FIREEModule::MakeModule(TSharedPtr<UE::NNE::FSharedModelData> SharedModelData, uint32 VmfbDataOffset, UE::NNERuntimeIREE::FModuleMetaData ModuleMetaData)
+	TSharedPtr<FIREEModule> FIREEModule::MakeModule(const FString& DirPath, const FString& VmfbFileName, const UE::NNERuntimeIREE::FModuleMetaData& ModuleMetaData)
 	{
-		check(SharedModelData->GetView().Num() > 0);
-		check(VmfbDataOffset < (uint32)SharedModelData->GetView().Num());
+		check(!VmfbFileName.IsEmpty());
 
 		TSharedPtr<FIREEInstance> Instance = FIREEInstance::GetInstance();
 		if (!Instance.IsValid())
@@ -430,8 +459,14 @@ namespace UE::NNERuntimeIREE
 			return TSharedPtr<FIREEModule>();
 		}
 
+		TSharedPtr<UE::NNE::FSharedModelData> ModelData = Instance->GetModelData(DirPath, VmfbFileName);
+		if (!ModelData.IsValid())
+		{
+			return TSharedPtr<FIREEModule>();
+		}
+
 		iree_vm_module_t* TempModule = nullptr;
-		if (!Instance->CreateModule(TConstArrayView<uint8>(&SharedModelData->GetView()[VmfbDataOffset], SharedModelData->GetView().Num() - VmfbDataOffset), &TempModule))
+		if (!Instance->CreateModule(ModelData->GetView(), &TempModule))
 		{
 			return TSharedPtr<FIREEModule>();
 		}
@@ -527,7 +562,7 @@ namespace UE::NNERuntimeIREE
 
 		TSharedPtr<FIREEModule> Result = MakeShared<FIREEModule>();
 		Result->Instance = Instance;
-		Result->ModelData = SharedModelData;
+		Result->ModelData = ModelData;
 		Result->Module = TempModule;
 		Result->MainFunction = LocalMainFunction;
 		Result->InputTensorDescs = InputTensorDescs;
@@ -630,7 +665,7 @@ namespace UE::NNERuntimeIREE
 		}
 
 		TConstArrayView<UE::NNE::FTensorDesc> InputTensorDescs = IREEModule->GetInputTensorDescs();
-		check(InputTensorDescs.Num() > 0);
+		check(!InputTensorDescs.IsEmpty());
 		bool bAllConcrete = true;
 		for (int32 i = 0; i < InputTensorDescs.Num(); i++)
 		{
@@ -646,7 +681,7 @@ namespace UE::NNERuntimeIREE
 		}
 
 		TConstArrayView<UE::NNE::FTensorDesc> OutputTensorDescs = IREEModule->GetOutputTensorDescs();
-		check(OutputTensorDescs.Num() > 0);
+		check(!OutputTensorDescs.IsEmpty());
 		bAllConcrete = true;
 		for (int32 i = 0; i < OutputTensorDescs.Num(); i++)
 		{
@@ -669,7 +704,7 @@ namespace UE::NNERuntimeIREE
 	int32 FIREESession::SetInputTensorShapes(TConstArrayView<UE::NNE::FTensorShape> InInputShapes)
 	{
 		check(Device.IsValid());
-		check(InInputShapes.Num() > 0);
+		check(!InInputShapes.IsEmpty());
 
 		iree_status_t Status = iree_ok_status();
 		check(iree_status_is_ok(Status));
@@ -735,6 +770,12 @@ namespace UE::NNERuntimeIREE
 		for (int32 i = 0; i < InInputBindings.Num(); i++)
 		{
 			check(InInputBindings[i].SizeInBytes == InputTensorShapes[i].Volume() * Module->GetInputTensorDescs()[i].GetElementByteSize());
+			
+			if (FMath::Modulo<uint64>((uint64)InInputBindings[i].Data, IREE_HAL_HEAP_BUFFER_ALIGNMENT) != 0)
+			{
+				UE_LOG(LogTemp, Error, TEXT("NNERuntimeIREECpu requires input- and output-buffer memory to be aligned with %d bytes"), IREE_HAL_HEAP_BUFFER_ALIGNMENT);
+				return -1;
+			}
 
 			iree_hal_buffer_view_t* TempBufferView;
 			iree_hal_buffer_params_t Params = { 0 };
@@ -915,7 +956,7 @@ namespace UE::NNERuntimeIREE
 
 	TSharedPtr<FIREELibrary> FIREELibrary::MakeLibrary(const FString& LibraryPath, const FString& LibraryName)
 	{
-		check(LibraryName.Len() > 0);
+		check(!LibraryName.IsEmpty());
 
 		TSharedPtr<FIREEInstance> IREEInstance = FIREEInstance::GetInstance();
 		if (!IREEInstance.IsValid())
@@ -964,7 +1005,7 @@ namespace UE::NNERuntimeIREE
 		check(Library.IsValid());
 		check(Library->Instance.IsValid());
 		check(Library->Library);
-		check(LibraryQueryFunctionName.Len() > 0);
+		check(!LibraryQueryFunctionName.IsEmpty());
 
 		if (Library->Devices.Contains(LibraryQueryFunctionName))
 		{
diff --git a/Engine/Plugins/Experimental/NNERuntimeIREE/Source/NNERuntimeIREE/Private/NNERuntimeIREECommon.h b/Engine/Plugins/Experimental/NNERuntimeIREE/Source/NNERuntimeIREE/Private/NNERuntimeIREECommon.h
--- a/Engine/Plugins/Experimental/NNERuntimeIREE/Source/NNERuntimeIREE/Private/NNERuntimeIREECommon.h
+++ b/Engine/Plugins/Experimental/NNERuntimeIREE/Source/NNERuntimeIREE/Private/NNERuntimeIREECommon.h
@@ -4,21 +4,21 @@
 
 #ifdef WITH_NNE_RUNTIME_IREE
 
-#include "Interfaces/ITargetPlatform.h"
-#include "Kismet/GameplayStatics.h"
+#include "CoreMinimal.h"
+#include "GenericPlatform/GenericPlatformProcess.h"
+#include "HAL/FileManager.h"
 #include "Misc/Paths.h"
 #include "Misc/ScopeLock.h"
 #include "NNEModelData.h"
-#include "NNERuntimeCPU.h"
+#include "NNERuntimeIREECpu.h"
 #include "NNERuntimeIREEMetaData.h"
 #include "NNETypes.h"
+#include "Serialization/Archive.h"
 
 #include "iree/runtime/api.h"
 
 namespace UE::NNERuntimeIREE
 {
-	FString GetTargetPlatformDisplayName(const ITargetPlatform* TargetPlatform);
-
 	class FIREEInstance;
 	class FIREELibrary;
 
@@ -28,7 +28,7 @@ namespace UE::NNERuntimeIREE
 		FIREEModule();
 		~FIREEModule();
 
-		static TSharedPtr<FIREEModule> MakeModule(TSharedPtr<UE::NNE::FSharedModelData> SharedModelData, uint32 VmfbDataOffset, UE::NNERuntimeIREE::FModuleMetaData ModuleMetaData);
+		static TSharedPtr<FIREEModule> MakeModule(const FString& DirPath, const FString& VmfbFileName, const UE::NNERuntimeIREE::FModuleMetaData& ModuleMetaData);
 		bool AppendToSession(iree_runtime_session_t* Session);
 
 		iree_vm_function_t GetMainFunction();
diff --git a/Engine/Plugins/Experimental/NNERuntimeIREE/Source/NNERuntimeIREE/Private/NNERuntimeIREECpu.cpp b/Engine/Plugins/Experimental/NNERuntimeIREE/Source/NNERuntimeIREE/Private/NNERuntimeIREECpu.cpp
--- a/Engine/Plugins/Experimental/NNERuntimeIREE/Source/NNERuntimeIREE/Private/NNERuntimeIREECpu.cpp
+++ b/Engine/Plugins/Experimental/NNERuntimeIREE/Source/NNERuntimeIREE/Private/NNERuntimeIREECpu.cpp
@@ -4,57 +4,39 @@
 
 #ifdef WITH_NNE_RUNTIME_IREE
 
-#include "EngineAnalytics.h"
-#include "GenericPlatform/GenericPlatformMisc.h"
+#if WITH_EDITOR
+#include "Containers/StringConv.h"
+#include "DerivedDataCache.h"
+#include "DerivedDataRequestOwner.h"
 #include "HAL/PlatformFileManager.h"
-#include "Interfaces/ITargetPlatform.h"
 #include "IO/IoHash.h"
+#include "Memory/SharedBuffer.h"
+#include "Misc/FileHelper.h"
+#endif // WITH_EDITOR
+
+#include "EngineAnalytics.h"
+#include "HAL/Platform.h"
+#include "Interfaces/ITargetPlatform.h"
 #include "Kismet/GameplayStatics.h"
 #include "Misc/App.h"
-#include "Misc/FileHelper.h"
 #include "Misc/Paths.h"
 #include "Misc/SecureHash.h"
-#include "NNERuntimeIREECommon.h"
 #include "NNERuntimeIREECpuModel.h"
 #include "Serialization/MemoryReader.h"
 
-#if WITH_EDITOR
-#include "Containers/StringFwd.h"
-#include "DerivedDataCacheKey.h"
-#include "DerivedDataCache.h"
-#include "DerivedDataRequestOwner.h"
-#include "Dom/JsonObject.h"
-#include "Interfaces/IPluginManager.h"
-#include "Internationalization/TextLocalizationResource.h"
-#include "Memory/SharedBuffer.h"
-#include "Misc/ConfigCacheIni.h"
-#include "Serialization/JsonReader.h"
-#include "Serialization/JsonSerializer.h"
-#include "Serialization/JsonSerializerMacros.h"
-#endif // WITH_EDITOR
-
 namespace UE::NNERuntimeIREECpu::Private
 {
 #if WITH_EDITOR
-	inline FString GetDDCRequestId(const FString& FileId, const FString& RuntimeName, const FString& ModelDataIdentifier)
-	{
-		//RuntimeName and FileId are embedded to the id to ensure no potential collision between the runtime/assets
-		return RuntimeName + "-" + FileId + "-" + ModelDataIdentifier + "-lib";
-	}
-
 	inline UE::DerivedData::FCacheKey CreateCacheKey(const FString& RequestId)
 	{
-		return { UE::DerivedData::FCacheBucket(TEXT("NNEModelData")), FIoHash::HashBuffer(MakeMemoryView(FTCHARToUTF8(RequestId))) };
+		return { UE::DerivedData::FCacheBucket(TEXT("NNEModelData")), FIoHash::HashBuffer(MakeMemoryView(StringCast<UTF8CHAR>(*RequestId))) };
 	}
 
-	inline void PutIntoDDC(const FGuid& FileId, const FString& RuntimeName, const FString& ModelDataIdentifier, const FSharedBuffer& Data)
+	inline void PutIntoDDC(const FString& RequestId, const FSharedBuffer& Data)
 	{
 		check(!Data.IsNull());
 		check(Data.GetSize() > 0);
 
-		FString FileIdStr = FileId.ToString(EGuidFormats::Digits);
-		FString RequestId = GetDDCRequestId(FileIdStr, RuntimeName, ModelDataIdentifier);
-
 		TArray<UE::DerivedData::FCachePutValueRequest> Requests;
 		Requests.SetNum(1);
 
@@ -67,11 +49,8 @@ namespace UE::NNERuntimeIREECpu::Private
 		BlockingPutOwner.Wait();
 	}
 
-	inline FSharedBuffer GetFromDDC(const FGuid& FileId, const FString& RuntimeName, const FString& ModelDataIdentifier)
+	inline FSharedBuffer GetFromDDC(const FString& RequestId)
 	{
-		FString FileIdStr = FileId.ToString(EGuidFormats::Digits);
-		FString RequestId = GetDDCRequestId(FileIdStr, RuntimeName, ModelDataIdentifier);
-
 		TArray<UE::DerivedData::FCacheGetValueRequest> Requests;
 		Requests.SetNum(1);
 
@@ -90,168 +69,40 @@ namespace UE::NNERuntimeIREECpu::Private
 		BlockingGetOwner.Wait();
 		return Result;
 	}
+#endif // WITH_EDITOR
 
-	struct FBuildConfig : FJsonSerializable
+	FString GetModelCpuDataIdentifier(const FString& RuntimeName, const FString& FileIdString, const FString& PlatformName, const FString& Architecture)
 	{
-		FString TargetPlatformName;
-		FString SharedLibraryExtension;
-		FString ModelCompilationCommand;
-		FString ModelCompilationArguments;
-		FString SharedLibraryCompilationCommand;
-		FString SharedLibraryCompilationArguments;
-
-		BEGIN_JSON_SERIALIZER
-			JSON_SERIALIZE("TargetPlatformName", TargetPlatformName);
-			JSON_SERIALIZE("SharedLibraryExtension", SharedLibraryExtension);
-			JSON_SERIALIZE("ModelCompilationCommand", ModelCompilationCommand);
-			JSON_SERIALIZE("ModelCompilationArguments", ModelCompilationArguments);
-			JSON_SERIALIZE("SharedLibraryCompilationCommand", SharedLibraryCompilationCommand);
-			JSON_SERIALIZE("SharedLibraryCompilationArguments", SharedLibraryCompilationArguments);
-		END_JSON_SERIALIZER
-	};
-
-	bool GetEnvVariable(const FString& Name, FString& Result)
-	{
-#if PLATFORM_WINDOWS
-		DWORD Size = GetEnvironmentVariableA(StringCast<ANSICHAR>(*Name).Get(), nullptr, 0);
-		if (Size > 0)
-		{
-			char* Data = new char[Size];
-			GetEnvironmentVariableA(StringCast<ANSICHAR>(*Name).Get(), Data, Size);
-			Result = FString(Data);
-			delete[] Data;
-			return true;
-		}
-		return false;
-#elif PLATFORM_LINUX
-		return false;
-#elif PLATFORM_MAC
-		return false;
-#else
-		return false;
-#endif
+		return RuntimeName + "-" + UNNERuntimeIREECpu::GUID.ToString(EGuidFormats::Digits) + "-" + FString::FromInt(UNNERuntimeIREECpu::Version) + "-" + FileIdString + "-" + PlatformName + (!Architecture.IsEmpty() ? ("-" + Architecture) : "");
 	}
 
-	bool ResolveEnvironmentVariables(FString& String)
+	bool GetModuleMetaData(const FString& FileType, TConstArrayView<uint8> FileData, UE::NNERuntimeIREE::FModuleMetaData* ModuleMetaData)
 	{
-		FString ResultString = String;
-		int32 StartIndex = String.Find("$ENV{", ESearchCase::CaseSensitive);
-		while (StartIndex != INDEX_NONE)
-		{
-			StartIndex += 5;
-			int32 EndIndex = String.Find("}", ESearchCase::CaseSensitive, ESearchDir::FromStart, StartIndex);
-			if (EndIndex > StartIndex)
-			{
-				FString EnvironmentVariableName = String.Mid(StartIndex, EndIndex - StartIndex);
-				FString EnvironmentVariableValue;
-				if (!GetEnvVariable(EnvironmentVariableName, EnvironmentVariableValue))
-				{
-					return false;
-				}
-				else
-				{
-					ResultString = ResultString.Replace(*(FString("$ENV{") + EnvironmentVariableName + FString("}")), *EnvironmentVariableValue, ESearchCase::CaseSensitive);
-				}
-			}
-			else
-			{
-				return false;
-			}
-			StartIndex = String.Find("$ENV{", ESearchCase::CaseSensitive, ESearchDir::FromStart, EndIndex);
-		}
-		String = ResultString;
-		return true;
+		check(ModuleMetaData);
+		FString FileDataString = "";
+		FileDataString.AppendChars((char*)FileData.GetData(), FileData.Num());
+		return ModuleMetaData->ParseFromString(FileDataString);
 	}
 
-	bool LoadBuildConfig(const FString& InTargetPlatformDisplayName, const FString& InBuildConfigFileName, FBuildConfig& OutBuildConfig)
+	FString GetIntermediateModelDirPath(const FString& PlatformName, const FString& ModelName)
 	{
-		IPlatformFile& PlatformFile = FPlatformFileManager::Get().GetPlatformFile();
-
-		FString PluginDir = FPaths::ConvertRelativePathToFull(*IPluginManager::Get().FindPlugin(UE_PLUGIN_NAME)->GetBaseDir());
-		TArray<FString> BuildConfigFilePaths =
-		{
-			FPaths::Combine(FPaths::ConvertRelativePathToFull(*IPluginManager::Get().FindPlugin(UE_PLUGIN_NAME)->GetBaseDir()), "Config", InBuildConfigFileName),
-			FPaths::Combine(FPaths::ConvertRelativePathToFull(*FPaths::EngineDir()), "Platforms", InTargetPlatformDisplayName, "Plugins", UE_PLUGIN_NAME, "Config", InBuildConfigFileName),
-			FPaths::Combine(FPaths::ConvertRelativePathToFull(*FPaths::EngineDir()), "Platforms", InTargetPlatformDisplayName, "Plugins", "Experimental", UE_PLUGIN_NAME, "Config", InBuildConfigFileName)
-		};
-
-		FString BuildConfigFileString = "";
-		for (FString BuildConfigFilePath : BuildConfigFilePaths)
-		{
-			if (PlatformFile.FileExists(*BuildConfigFilePath)) 
-			{
-				if (FFileHelper::LoadFileToString(BuildConfigFileString, *BuildConfigFilePath))
-				{
-					TSharedRef<TJsonReader<TCHAR>> JsonReader = TJsonReaderFactory<TCHAR>::Create(BuildConfigFileString);
-					TSharedPtr<FJsonObject> JsonObject = MakeShareable(new FJsonObject);
-					if (FJsonSerializer::Deserialize(JsonReader, JsonObject) && JsonObject.IsValid() && OutBuildConfig.FromJson(JsonObject))
-					{
-						if (ResolveEnvironmentVariables(OutBuildConfig.ModelCompilationCommand))
-						{
-							if (ResolveEnvironmentVariables(OutBuildConfig.SharedLibraryCompilationCommand))
-							{
-								return true;
-							}
-							else
-							{
-								UE_LOG(LogTemp, Warning, TEXT("UNNERuntimeIREECpu could not resolve all environment variables in %s"), *OutBuildConfig.SharedLibraryCompilationCommand);
-								return false;
-							}
-						}
-						else
-						{
-							UE_LOG(LogTemp, Warning, TEXT("UNNERuntimeIREECpu could not resolve all environment variables in %s"), *OutBuildConfig.ModelCompilationCommand);
-							return false;
-						}
-					}
-				}
-				else
-				{
-					UE_LOG(LogTemp, Warning, TEXT("UNNERuntimeIREECpu could not read build config file %s"), *BuildConfigFilePath);
-					return false;
-				}
-			}
-		}
-
-		UE_LOG(LogTemp, Warning, TEXT("UNNERuntimeIREECpu could not find any build config file"));
-		return false;
+		return FPaths::Combine("Intermediate", "Build", PlatformName, UE_PLUGIN_NAME, ModelName);
 	}
-#endif // WITH_EDITOR
 
-	FString GetModelCpuDataIdentifier(const FString& FileIdString, const FString& PlatformDisplayName)
+	FString GetStagedModelDirPath(const FString& PlatformName)
 	{
-		return FileIdString + "-" + UNNERuntimeIREECpu::GUID.ToString(EGuidFormats::Digits) + "-" + FString::FromInt(UNNERuntimeIREECpu::Version) + "-" + FString::FromInt(UNNERuntimeIREECpu::MemoryAlignment) + "-" + PlatformDisplayName;
+		return FPaths::Combine("Saved", "Cooked", PlatformName, "Engine", "Plugins", UE_PLUGIN_NAME, "Binaries");
 	}
 
-	FString GetLibraryEntryPointName(const FString& HeaderString)
+	FString GetPackagedModelDirPath(const FString& PlatformName)
 	{
-		FString SearchString = "iree_hal_executable_library_header_t**";
-		int32 Start = HeaderString.Find(SearchString);
-		if (Start < 0)
-		{
-			return "";
-		}
-		Start += SearchString.Len();
-		int32 End = HeaderString.Find("(", ESearchCase::CaseSensitive, ESearchDir::FromStart, Start);
-		if (End <= Start)
-		{
-			return "";
-		}
-		return HeaderString.Mid(Start, End - Start).TrimStartAndEnd();
-	}
-
-	bool GetModuleMetaData(const FString& FileType, TConstArrayView<uint8> FileData, UE::NNERuntimeIREE::FModuleMetaData* ModuleMetaData)
-	{
-		check(ModuleMetaData);
-		FString FileDataString = "";
-		FileDataString.AppendChars((char*)FileData.GetData(), FileData.Num());
-		return ModuleMetaData->ParseFromString(FileDataString);
+		FString PlatformNameShort = PlatformName.Equals("Windows") ? "Win64" : PlatformName;
+		return FPaths::Combine("Binaries", PlatformNameShort, UE_PLUGIN_NAME);
 	}
 } // UE::NNERuntimeIREECpu::Private
 
 FGuid UNNERuntimeIREECpu::GUID = FGuid((int32)'I', (int32)'C', (int32)'P', (int32)'U');
-int32 UNNERuntimeIREECpu::Version = 0x00000001;
-uint32 UNNERuntimeIREECpu::MemoryAlignment = IREE_HAL_HEAP_BUFFER_ALIGNMENT;
+int32 UNNERuntimeIREECpu::Version = 0x00000002;
 
 FString UNNERuntimeIREECpu::GetRuntimeName() const
 {
@@ -270,152 +121,95 @@ bool UNNERuntimeIREECpu::CanCreateModelData(FString FileType, TConstArrayView<ui
 TSharedPtr<UE::NNE::FSharedModelData> UNNERuntimeIREECpu::CreateModelData(FString FileType, TConstArrayView<uint8> FileData, FGuid FileId, const ITargetPlatform* TargetPlatform)
 {
 #if WITH_EDITOR
-	FString TargetPlatformDisplayName = UE::NNERuntimeIREE::GetTargetPlatformDisplayName(TargetPlatform);
+	using namespace UE::NNERuntimeIREECpu::Private;
+
+	FString TargetPlatformName = TargetPlatform ? TargetPlatform->IniPlatformName() : UGameplayStatics::GetPlatformName();
 	if (!CanCreateModelData(FileType, FileData, FileId, TargetPlatform))
 	{
-		UE_LOG(LogTemp, Warning, TEXT("UNNERuntimeIREECpu cannot create the model data with id %s (Filetype: %s) for platform %s"), *FileId.ToString(EGuidFormats::Digits).ToLower(), *FileType, *TargetPlatformDisplayName);
+		UE_LOG(LogTemp, Warning, TEXT("UNNERuntimeIREECpu cannot create the model data with id %s (Filetype: %s) for platform %s"), *FileId.ToString(EGuidFormats::Digits).ToLower(), *FileType, *TargetPlatformName);
 		return TSharedPtr<UE::NNE::FSharedModelData>();
 	}
 
 	FConfigFile ConfigFile;
 	FString ConfigFilePath;
-	GetUpdatedPlatformConfig(TargetPlatformDisplayName, ConfigFile, ConfigFilePath);
+	GetUpdatedPlatformConfig(TargetPlatformName, ConfigFile, ConfigFilePath);
 	if (ConfigFile.Dirty)
 	{
-		UE_LOG(LogTemp, Warning, TEXT("UNNERuntimeIREECpu could not find the required settings in config file %s. Please make the file writeable and re-start the editor or manually add the required staging settings or models will not work in packaged builds for platform %s!"), *ConfigFilePath, *TargetPlatformDisplayName);
+		UE_LOG(LogTemp, Warning, TEXT("UNNERuntimeIREECpu could not find the required settings in config file %s. Please make the file writeable and re-start the editor or manually add the required staging settings or models will not work in packaged builds for platform %s!"), *ConfigFilePath, *TargetPlatformName);
 	}
 
-	FString BuildConfigFileName = FString("IREE_") + NNE_RUNTIME_IREE_PLATFORM_NAME + "_To_" + TargetPlatformDisplayName + ".json";
-	UE::NNERuntimeIREECpu::Private::FBuildConfig BuildConfig;
-	if (!UE::NNERuntimeIREECpu::Private::LoadBuildConfig(TargetPlatformDisplayName, BuildConfigFileName, BuildConfig))
+	TUniquePtr<FNNERuntimeIREECpuCompiler> Compiler = FNNERuntimeIREECpuCompiler::Make(TargetPlatformName);
+	if (!Compiler.IsValid())
 	{
-		UE_LOG(LogTemp, Warning, TEXT("UNNERuntimeIREECpu failed to find and load the build config file %s"), *BuildConfigFileName);
+		UE_LOG(LogTemp, Warning, TEXT("UNNERuntimeIREECpu failed to create a compiler to compile for platform %s"), *TargetPlatformName);
 		return TSharedPtr<UE::NNE::FSharedModelData>();
 	}
-
+	
 	FString FileIdString = FileId.ToString(EGuidFormats::Digits).ToLower();
-	FString PluginDir = FPaths::ConvertRelativePathToFull(*IPluginManager::Get().FindPlugin(UE_PLUGIN_NAME)->GetBaseDir());
-	FString IntermediateModelDir = FPaths::ConvertRelativePathToFull(FPaths::Combine(FPaths::ProjectDir(), GetIntermediateModelDirPath(BuildConfig.TargetPlatformName, FileIdString)));
-	FString ObjectPath = FPaths::Combine(IntermediateModelDir, FileIdString + ".o");
-	FString AdditionalOutputPath = FPaths::Combine(IntermediateModelDir, FileIdString + ".vmfb");
-	FString InputPath = FPaths::Combine(IntermediateModelDir, FileIdString + "." + FileType);
-	FString SharedLibName = FileIdString + "." + BuildConfig.SharedLibraryExtension;
-	FString SharedLibPath = FPaths::Combine(IntermediateModelDir, SharedLibName);
-
+	FString IntermediateDir = FPaths::ConvertRelativePathToFull(FPaths::Combine(FPaths::ProjectDir(), GetIntermediateModelDirPath(TargetPlatformName, FileIdString)));
 	IPlatformFile& PlatformFile = FPlatformFileManager::Get().GetPlatformFile();
-	PlatformFile.DeleteDirectoryRecursively(*IntermediateModelDir);
-	PlatformFile.CreateDirectoryTree(*IntermediateModelDir);
-	FFileHelper::SaveArrayToFile(FileData, *InputPath);
-
-	FString ModelCompilationCommand = BuildConfig.ModelCompilationCommand.Replace(*FString("${PLUGIN_DIR}"), *PluginDir);
-	FString ModelCompilationArguments = BuildConfig.ModelCompilationArguments.Replace(*FString("${OBJECT_PATH}"), *ObjectPath).Replace(*FString("${ADDITIONAL_OUTPUT_PATH}"), *AdditionalOutputPath).Replace(*FString("${INPUT_PATH}"), *InputPath);
-
-	FProcHandle ModelCompilationProcHandle = FPlatformProcess::CreateProc(*ModelCompilationCommand, *ModelCompilationArguments, false, true, true, nullptr, 0, nullptr, nullptr, nullptr);
-	FPlatformProcess::WaitForProc(ModelCompilationProcHandle);
+	PlatformFile.DeleteDirectoryRecursively(*IntermediateDir);
+	PlatformFile.CreateDirectoryTree(*IntermediateDir);
 
-	if ((BuildConfig.ModelCompilationArguments.Contains(FString("${ADDITIONAL_OUTPUT_PATH}")) && !PlatformFile.FileExists(*AdditionalOutputPath)) ||
-		(BuildConfig.ModelCompilationArguments.Contains(FString("${OBJECT_PATH}")) && !PlatformFile.FileExists(*ObjectPath)))
+	TArray<FIREECompilerResult> CompilerResults;
+	FString StagingDir = FPaths::ConvertRelativePathToFull(FPaths::Combine(FPaths::ProjectDir(), GetStagedModelDirPath(TargetPlatformName)));
+	if (!Compiler->CompileMlir(FileData, FileIdString, IntermediateDir, StagingDir, CompilerResults))
 	{
-		UE_LOG(LogTemp, Warning, TEXT("UNNERuntimeIREECpu failed to compile the model \"%s\" using the command:"), *InputPath);
-		UE_LOG(LogTemp, Warning, TEXT("\"%s\" %s"), *ModelCompilationCommand, *ModelCompilationArguments);
+		UE_LOG(LogTemp, Warning, TEXT("UNNERuntimeIREECpu failed to compile model %s"), *FileIdString);
 		return TSharedPtr<UE::NNE::FSharedModelData>();
 	}
 
-	if (BuildConfig.SharedLibraryCompilationCommand.Len() > 0)
+	for (int32 i = 0; i < CompilerResults.Num(); i++)
 	{
-		FString SharedLibraryCompilationCommand = BuildConfig.SharedLibraryCompilationCommand;
-		FString SharedLibraryCompilationArguments = BuildConfig.SharedLibraryCompilationArguments;
-#ifdef NNE_RUNTIME_IREE_WIN_TOOLCHAIN_PATH
-		SharedLibraryCompilationCommand = SharedLibraryCompilationCommand.Replace(*FString("${WIN_TOOLCHAIN_PATH}"), *FString(NNE_RUNTIME_IREE_WIN_TOOLCHAIN_PATH));
-		SharedLibraryCompilationArguments = SharedLibraryCompilationArguments.Replace(*FString("${WIN_TOOLCHAIN_PATH}"), *FString(NNE_RUNTIME_IREE_WIN_TOOLCHAIN_PATH));
-#endif
-#ifdef NNE_RUNTIME_IREE_WIN_SDK_LIB_PATH
-		SharedLibraryCompilationArguments = SharedLibraryCompilationArguments.Replace(*FString("${WIN_SDK_LIB_PATH}"), *FString(NNE_RUNTIME_IREE_WIN_SDK_LIB_PATH));
-#endif
-#ifdef NNE_RUNTIME_IREE_LINUX_COMPILER
-		SharedLibraryCompilationCommand = SharedLibraryCompilationCommand.Replace(*FString("${LINUX_COMPILER}"), *FString(NNE_RUNTIME_IREE_LINUX_COMPILER));
-#endif
-		FString FinalSharedLibraryCompilationArguments = SharedLibraryCompilationArguments.Replace(*FString("${OBJECT_PATH}"), *ObjectPath).Replace(*FString("${SHARED_LIB_PATH}"), *SharedLibPath);
-		FProcHandle SharedLibraryCompilationProcHandle = FPlatformProcess::CreateProc(*SharedLibraryCompilationCommand, *FinalSharedLibraryCompilationArguments, false, true, true, nullptr, 0, nullptr, nullptr, nullptr);
-		FPlatformProcess::WaitForProc(SharedLibraryCompilationProcHandle);
-		if (!PlatformFile.FileExists(*SharedLibPath))
+		TArray<uint8> SharedLibData;
+		FString StagedSharedLibPath = FPaths::Combine(StagingDir, CompilerResults[i].RelativeDirPath, CompilerResults[i].SharedLibraryFileName);
+		if (!FFileHelper::LoadFileToArray(SharedLibData, *StagedSharedLibPath) || SharedLibData.IsEmpty())
 		{
-			UE_LOG(LogTemp, Warning, TEXT("UNNERuntimeIREECpu failed to compile the shared library for the object \"%s\" using the command:"), *ObjectPath);
-			UE_LOG(LogTemp, Warning, TEXT("\"%s\" %s"), *SharedLibraryCompilationCommand, *FinalSharedLibraryCompilationArguments);
+			UE_LOG(LogTemp, Warning, TEXT("UNNERuntimeIREECpu could not read the shared library \"%s\""), *StagedSharedLibPath);
 			return TSharedPtr<UE::NNE::FSharedModelData>();
 		}
+		FSharedBuffer SharedLibBuffer = MakeSharedBufferFromArray(MoveTemp(SharedLibData));
+		PutIntoDDC(GetModelCpuDataIdentifier(GetRuntimeName(), FileIdString, TargetPlatformName, CompilerResults[i].Architecture) + "-lib", SharedLibBuffer);
 
-		FString CookedModelPath = FPaths::ConvertRelativePathToFull(FPaths::Combine(FPaths::ProjectDir(), GetCookedModelDirPath(TargetPlatformDisplayName), SharedLibName));
-		IFileManager::Get().Copy(*CookedModelPath, *SharedLibPath);
-
-		TArray<uint8> SharedLibData;
-		if (!FFileHelper::LoadFileToArray(SharedLibData, *CookedModelPath) || SharedLibData.Num() < 1)
+		TArray<uint8> VmfbData;
+		FString StagedVmfbPath = FPaths::Combine(StagingDir, CompilerResults[i].RelativeDirPath, CompilerResults[i].VmfbFileName);
+		if (!FFileHelper::LoadFileToArray(VmfbData, *StagedVmfbPath) || VmfbData.IsEmpty())
 		{
-			UE_LOG(LogTemp, Warning, TEXT("UNNERuntimeIREECpu could not read the shared library \"%s\""), *CookedModelPath);
+			UE_LOG(LogTemp, Warning, TEXT("UNNERuntimeIREECpu could not read the vmfb data \"%s\""), *StagedVmfbPath);
 			return TSharedPtr<UE::NNE::FSharedModelData>();
 		}
-		FSharedBuffer SharedBuffer = MakeSharedBufferFromArray(MoveTemp(SharedLibData));
-		UE::NNERuntimeIREECpu::Private::PutIntoDDC(FileId, GetRuntimeName(), UE::NNERuntimeIREECpu::Private::GetModelCpuDataIdentifier(FileIdString, TargetPlatformDisplayName), SharedBuffer);
+		FSharedBuffer VmfbBuffer = MakeSharedBufferFromArray(MoveTemp(VmfbData));
+		PutIntoDDC(GetModelCpuDataIdentifier(GetRuntimeName(), FileIdString, TargetPlatformName, CompilerResults[i].Architecture) + "-vmfb", VmfbBuffer);
 	}
-
-	FString HeaderPath = FPaths::Combine(IntermediateModelDir, FileIdString + ".h");
-	if (!PlatformFile.FileExists(*HeaderPath))
-	{
-		UE_LOG(LogTemp, Warning, TEXT("UNNERuntimeIREECpu could not find the model header \"%s\""), *HeaderPath);
-		return TSharedPtr<UE::NNE::FSharedModelData>();
-	}
-	FString HeaderString;
-	if (!FFileHelper::LoadFileToString(HeaderString, *HeaderPath) || HeaderString.Len() < 1)
-	{
-		UE_LOG(LogTemp, Warning, TEXT("UNNERuntimeIREECpu could not read the model header \"%s\""), *HeaderPath);
-		return TSharedPtr<UE::NNE::FSharedModelData>();
-	}
-	FString LibraryQueryFunctionName = UE::NNERuntimeIREECpu::Private::GetLibraryEntryPointName(HeaderString);
-	if(LibraryQueryFunctionName.Len() < 1)
-	{
-		UE_LOG(LogTemp, Warning, TEXT("UNNERuntimeIREECpu could not find the entry point in model header \"%s\""), *HeaderPath);
-		return TSharedPtr<UE::NNE::FSharedModelData>();
-	}
-
-	TArray<uint8> AdditionalOutputData;
-	if (!FFileHelper::LoadFileToArray(AdditionalOutputData, *AdditionalOutputPath))
-	{
-		UE_LOG(LogTemp, Warning, TEXT("UNNERuntimeIREECpu could not load the model virtual machine buffer \"%s\""), *AdditionalOutputPath);
-		return TSharedPtr<UE::NNE::FSharedModelData>();
-	}
-
+	
 	UE::NNERuntimeIREE::FModuleMetaData ModuleMetaData;
-	if (!UE::NNERuntimeIREECpu::Private::GetModuleMetaData(FileType, FileData, &ModuleMetaData))
+	if (!GetModuleMetaData(FileType, FileData, &ModuleMetaData))
 	{
 		UE_LOG(LogTemp, Warning, TEXT("UNNERuntimeIREECpu could not get the module meta data"));
 		return TSharedPtr<UE::NNE::FSharedModelData>();
 	}
 	FString ModuleDataString = ModuleMetaData.ToJson(false);
 
-	TArray<uint8> Header;
-	FMemoryWriter Writer(Header);
+	TArray<uint8> ResultData;
+	FMemoryWriter Writer(ResultData);
 	Writer << UNNERuntimeIREECpu::GUID;
 	Writer << UNNERuntimeIREECpu::Version;
 	Writer << FileId;
-	Writer << LibraryQueryFunctionName;
+
 	Writer << ModuleDataString;
 
-	uint32 AdditionalDataOffset = 0;
-	uint32 HeaderSize = (uint32)Header.Num() + sizeof(AdditionalDataOffset);
-	check(MemoryAlignment > 0);
-	AdditionalDataOffset = uint32(FMath::CeilToDouble(double(HeaderSize) / double(MemoryAlignment))) * MemoryAlignment;
-	check(AdditionalDataOffset > 0);
-	check(HeaderSize < AdditionalDataOffset);
-	check(AdditionalDataOffset % MemoryAlignment == 0);
-
-	Writer << AdditionalDataOffset;
-	check(HeaderSize == (uint32)Header.Num());
-
-	SIZE_T ResultDataSize = AdditionalDataOffset + AdditionalOutputData.Num();
-	void* ResultData = FMemory::Malloc(ResultDataSize, MemoryAlignment);
-	FMemory::Memcpy(ResultData, (void*)Header.GetData(), Header.Num());
-	FMemory::Memcpy((void*)&((uint8*)ResultData)[AdditionalDataOffset], (void*)AdditionalOutputData.GetData(), AdditionalOutputData.Num());
-	return MakeShared<UE::NNE::FSharedModelData>(FSharedBuffer::TakeOwnership(ResultData, ResultDataSize, FMemory::Free), MemoryAlignment);
+	int32 NumArchitectures = CompilerResults.Num();
+	Writer << NumArchitectures;
+	for (int32 i = 0; i < NumArchitectures; i++)
+	{
+		Writer << CompilerResults[i].Architecture;
+		Writer << CompilerResults[i].RelativeDirPath;
+		Writer << CompilerResults[i].SharedLibraryFileName;
+		Writer << CompilerResults[i].VmfbFileName;
+		Writer << CompilerResults[i].SharedLibraryEntryPointName;
+	}
+
+	return MakeShared<UE::NNE::FSharedModelData>(MakeSharedBufferFromArray(MoveTemp(ResultData)), 0);
 #else
 	return TSharedPtr<UE::NNE::FSharedModelData>();
 #endif // WITH_EDITOR
@@ -423,8 +217,9 @@ TSharedPtr<UE::NNE::FSharedModelData> UNNERuntimeIREECpu::CreateModelData(FStrin
 
 FString UNNERuntimeIREECpu::GetModelDataIdentifier(FString FileType, TConstArrayView<uint8> FileData, FGuid FileId, const ITargetPlatform* TargetPlatform)
 {
-	FString PlatformDisplayName = UE::NNERuntimeIREE::GetTargetPlatformDisplayName(TargetPlatform);
-	return UE::NNERuntimeIREECpu::Private::GetModelCpuDataIdentifier(FileId.ToString(EGuidFormats::Digits), PlatformDisplayName);
+	// Leave architecture blank as there is only one model data for all architectures of a given platform
+	FString PlatformName = TargetPlatform ? TargetPlatform->IniPlatformName() : UGameplayStatics::GetPlatformName();
+	return UE::NNERuntimeIREECpu::Private::GetModelCpuDataIdentifier(GetRuntimeName(), FileId.ToString(EGuidFormats::Digits), PlatformName, "");
 }
 
 bool UNNERuntimeIREECpu::CanCreateModelCPU(TObjectPtr<UNNEModelData> ModelData) const
@@ -454,31 +249,84 @@ TSharedPtr<UE::NNE::IModelCPU> UNNERuntimeIREECpu::CreateModelCPU(TObjectPtr<UNN
 {
 	check(ModelData != nullptr);
 
+	using namespace UE::NNERuntimeIREECpu::Private;
+
 	if (!CanCreateModelCPU(ModelData))
 	{
 		UE_LOG(LogTemp, Warning, TEXT("UNNERuntimeIREECpu cannot create a model from the model data with id %s"), *ModelData->GetFileId().ToString(EGuidFormats::Digits));
 		return TSharedPtr<UE::NNE::IModelCPU>();
 	}
 
+	FString CurrentArchitecture = "";
+#ifdef PLATFORM_CPU_X86_FAMILY
+	CurrentArchitecture = "x86_64";
+#elif PLATFORM_CPU_ARM_FAMILY
+	CurrentArchitecture = "arm64";
+#endif
+
 	TSharedPtr<UE::NNE::FSharedModelData> SharedData = ModelData->GetModelData(GetRuntimeName());
 	check(SharedData.IsValid());
 	TConstArrayView<uint8> SharedDataView = SharedData->GetView();
 	FMemoryReaderView Reader(SharedDataView);
-	FGuid DataGuid;
+	FGuid DataGuid = FGuid();
 	Reader << DataGuid;
-	int32 VersionVersion;
+	int32 VersionVersion = 0;
 	Reader << VersionVersion;
-	FGuid FileId;
+	FGuid FileId = FGuid();
 	Reader << FileId;
-	FString LibraryQueryFunctionName;
-	Reader << LibraryQueryFunctionName;
-	FString ModuleDataString;
+
+	FString ModuleDataString = "";
 	Reader << ModuleDataString;
-	uint32 VmfbDataOffset;
-	Reader << VmfbDataOffset;
+
+	int32 NumArchitectures = 0;
+	Reader << NumArchitectures;
+
+	bool bFound = false;
+	FString Architecture = "";
+	FString RelativeDirPath = "";
+	FString SharedLibraryFileName = "";
+	FString VmfbFileName = "";
+	FString SharedLibraryEntryPointName = "";
+	for (int32 i = 0; i < NumArchitectures; i++)
+	{
+		FString TmpArchitecture = "";
+		Reader << TmpArchitecture;
+		FString TmpRelativeDirPath = "";
+		Reader << TmpRelativeDirPath;
+		FString TmpSharedLibraryFileName = "";
+		Reader << TmpSharedLibraryFileName;
+		FString TmpVmfbFileName = "";
+		Reader << TmpVmfbFileName;
+		FString TmpSharedLibraryEntryPointName = "";
+		Reader << TmpSharedLibraryEntryPointName;
+
+		if (TmpArchitecture.IsEmpty() && !bFound)
+		{
+			Architecture = TmpArchitecture;
+			RelativeDirPath = TmpRelativeDirPath;
+			SharedLibraryFileName = TmpSharedLibraryFileName;
+			VmfbFileName = TmpVmfbFileName;
+			SharedLibraryEntryPointName = TmpSharedLibraryEntryPointName;
+			bFound = true;
+		}
+		else if (TmpArchitecture.Equals(CurrentArchitecture))
+		{
+			Architecture = TmpArchitecture;
+			RelativeDirPath = TmpRelativeDirPath;
+			SharedLibraryFileName = TmpSharedLibraryFileName;
+			VmfbFileName = TmpVmfbFileName;
+			SharedLibraryEntryPointName = TmpSharedLibraryEntryPointName;
+			bFound = true;
+		}
+	}
+	if (!bFound)
+	{
+		UE_LOG(LogTemp, Warning, TEXT("UNNERuntimeIREECpu failed to find a matching architecture for \'%s\'"), *CurrentArchitecture);
+		return TSharedPtr<UE::NNE::IModelCPU>();
+	}
 
 	UE::NNERuntimeIREE::FModuleMetaData ModuleMetaData;
-	if (ModuleDataString.Len() > 0 && !ModuleMetaData.FromJson(ModuleDataString))
+	if (!ModuleDataString.IsEmpty() && !ModuleMetaData.FromJson(ModuleDataString))
 	{
 		UE_LOG(LogTemp, Warning, TEXT("UNNERuntimeIREECpu failed to parse the module meta data"));
 		return TSharedPtr<UE::NNE::IModelCPU>();
@@ -486,31 +334,42 @@ TSharedPtr<UE::NNE::IModelCPU> UNNERuntimeIREECpu::CreateModelCPU(TObjectPtr<UNN
 
 	FString FileIdString = FileId.ToString(EGuidFormats::Digits).ToLower();
 #if WITH_EDITOR
-	FString SharedLibDirPath = FPaths::ConvertRelativePathToFull(FPaths::Combine(FPaths::ProjectDir(), GetIntermediateModelDirPath(NNE_RUNTIME_IREE_PLATFORM_NAME, FileIdString)));
-	FString SharedLibName = FileIdString + "." + NNE_RUNTIME_IREE_PLATFORM_SHARED_LIB_EXTENSION;
+	FString SharedLibraryDirPath = FPaths::ConvertRelativePathToFull(FPaths::Combine(FPaths::ProjectDir(), GetIntermediateModelDirPath(UGameplayStatics::GetPlatformName(), FileIdString), RelativeDirPath));
 #else
-	FString SharedLibDirPath = FPaths::ConvertRelativePathToFull(FPaths::Combine(FPaths::ProjectDir(), GetPackagedModelDirPath(NNE_RUNTIME_IREE_PLATFORM_NAME)));
-	FString SharedLibName = FileIdString + "." + NNE_RUNTIME_IREE_PLATFORM_SHARED_LIB_EXTENSION;
+	FString SharedLibraryDirPath = FPaths::ConvertRelativePathToFull(FPaths::Combine(FPaths::ProjectDir(), GetPackagedModelDirPath(UGameplayStatics::GetPlatformName()), RelativeDirPath));
 #endif // WITH_EDITOR
 
 #if WITH_EDITOR
 	IPlatformFile& PlatformFile = FPlatformFileManager::Get().GetPlatformFile();
-	FString SharedLibPath = FPaths::Combine(SharedLibDirPath, SharedLibName);
-	if (!PlatformFile.FileExists(*SharedLibPath))
+	FString SharedLibraryFilePath = FPaths::Combine(SharedLibraryDirPath, SharedLibraryFileName);
+	if (!PlatformFile.FileExists(*SharedLibraryFilePath))
 	{
-		FSharedBuffer SharedBuffer = UE::NNERuntimeIREECpu::Private::GetFromDDC(FileId, GetRuntimeName(), UE::NNERuntimeIREECpu::Private::GetModelCpuDataIdentifier(FileIdString, NNE_RUNTIME_IREE_PLATFORM_DISPLAY_NAME));
+		FSharedBuffer SharedBuffer = GetFromDDC(GetModelCpuDataIdentifier(GetRuntimeName(), FileIdString, UGameplayStatics::GetPlatformName(), Architecture) + "-lib");
 		if (SharedBuffer.GetSize() <= 0)
 		{
-			UE_LOG(LogTemp, Warning, TEXT("UNNERuntimeIREECpu could not fetch the shared library %s from DDC"), *SharedLibName);
+			UE_LOG(LogTemp, Warning, TEXT("UNNERuntimeIREECpu could not fetch the shared library %s from DDC"), *SharedLibraryFileName);
 			return TSharedPtr<UE::NNE::IModelCPU>();
 		}
-		PlatformFile.CreateDirectoryTree(*SharedLibDirPath);
-		FFileHelper::SaveArrayToFile(TConstArrayView<uint8>((uint8*)SharedBuffer.GetData(), SharedBuffer.GetSize()), *SharedLibPath);
+		PlatformFile.CreateDirectoryTree(*SharedLibraryDirPath);
+		FFileHelper::SaveArrayToFile(TConstArrayView<uint8>((uint8*)SharedBuffer.GetData(), SharedBuffer.GetSize()), *SharedLibraryFilePath);
+	}
+
+	FString VmfbFilePath = FPaths::Combine(SharedLibraryDirPath, VmfbFileName);
+	if (!PlatformFile.FileExists(*VmfbFilePath))
+	{
+		FSharedBuffer SharedBuffer = GetFromDDC(GetModelCpuDataIdentifier(GetRuntimeName(), FileIdString, UGameplayStatics::GetPlatformName(), Architecture) + "-vmfb");
+		if (SharedBuffer.GetSize() <= 0)
+		{
+			UE_LOG(LogTemp, Warning, TEXT("UNNERuntimeIREECpu could not fetch the vmfb %s from DDC"), *VmfbFileName);
+			return TSharedPtr<UE::NNE::IModelCPU>();
+		}
+		PlatformFile.CreateDirectoryTree(*SharedLibraryDirPath);
+		FFileHelper::SaveArrayToFile(TConstArrayView<uint8>((uint8*)SharedBuffer.GetData(), SharedBuffer.GetSize()), *VmfbFilePath);
 	}
 #endif // WITH_EDITOR
 
 	TUniquePtr<UE::NNERuntimeIREECpu::FModel> Model = MakeUnique<UE::NNERuntimeIREECpu::FModel>();
-	if (!Model->Init(SharedData, VmfbDataOffset, ModuleMetaData, SharedLibDirPath, SharedLibName, LibraryQueryFunctionName))
+	if (!Model->Init(SharedLibraryDirPath, SharedLibraryFileName, VmfbFileName, SharedLibraryEntryPointName, ModuleMetaData))
 	{
 		UE_LOG(LogTemp, Warning, TEXT("UNNERuntimeIREECpu could not initialize the model created from model data with id %s"), *FileIdString);
 		return TSharedPtr<UE::NNE::IModelCPU>();
@@ -529,22 +388,6 @@ TSharedPtr<UE::NNE::IModelCPU> UNNERuntimeIREECpu::CreateModelCPU(TObjectPtr<UNN
 	return TSharedPtr<UE::NNE::IModelCPU>(static_cast<UE::NNE::IModelCPU*>(Model.Release()));
 }
 
-FString UNNERuntimeIREECpu::GetIntermediateModelDirPath(const FString& PlatformName, const FString& FileIdString)
-{
-	return FPaths::Combine("Intermediate", "Build", PlatformName, UE_PLUGIN_NAME, FileIdString);
-}
-
-FString UNNERuntimeIREECpu::GetCookedModelDirPath(const FString& PlatformName)
-{
-	return FPaths::Combine("Saved", "Cooked", PlatformName, "Engine", "Plugins", UE_PLUGIN_NAME, "Binaries");
-}
-
-FString UNNERuntimeIREECpu::GetPackagedModelDirPath(const FString& PlatformName)
-{
-	FString PlatformNameShort = PlatformName.Equals("Windows") ? "Win64" : PlatformName;
-	return FPaths::Combine("Binaries", PlatformNameShort, UE_PLUGIN_NAME);
-}
-
 void UNNERuntimeIREECpu::GetUpdatedPlatformConfig(const FString& PlatformName, FConfigFile& ConfigFile, FString& ConfigFilePath)
 { 
 	FString ConfigFolderPath = FPaths::ConvertRelativePathToFull(FPaths::ProjectConfigDir());
@@ -552,11 +395,11 @@ void UNNERuntimeIREECpu::GetUpdatedPlatformConfig(const FString& PlatformName, F
 
 	ConfigFile.Read(ConfigFilePath);
 
-	FString CookingPath = FString("/") + GetCookedModelDirPath(PlatformName);
-	FString PackagingPath = FString("/") + GetPackagedModelDirPath(PlatformName);
+	FString StagingPath = FString("/") + UE::NNERuntimeIREECpu::Private::GetStagedModelDirPath(PlatformName);
+	FString PackagingPath = FString("/") + UE::NNERuntimeIREECpu::Private::GetPackagedModelDirPath(PlatformName);
 
-	ConfigFile.AddUniqueToSection(TEXT("/Script/UnrealEd.ProjectPackagingSettings"), TEXT("+DirectoriesToAlwaysStageAsNonUFS"), FString("(Path=\"..") + CookingPath + FString("\")"));
-	ConfigFile.AddUniqueToSection(TEXT("Staging"), TEXT("+RemapDirectories"), FString("(From=\"") + FApp::GetProjectName() + CookingPath + FString("\", To=\"") + FApp::GetProjectName() + PackagingPath + FString("\")"));
+	ConfigFile.AddUniqueToSection(TEXT("/Script/UnrealEd.ProjectPackagingSettings"), TEXT("+DirectoriesToAlwaysStageAsNonUFS"), FString("(Path=\"..") + StagingPath + FString("\")"));
+	ConfigFile.AddUniqueToSection(TEXT("Staging"), TEXT("+RemapDirectories"), FString("(From=\"") + FApp::GetProjectName() + StagingPath + FString("\", To=\"") + FApp::GetProjectName() + PackagingPath + FString("\")"));
 	ConfigFile.AddUniqueToSection(TEXT("Staging"), TEXT("+AllowedDirectories"), FApp::GetProjectName() + PackagingPath);
 }
 
@@ -571,10 +414,6 @@ FString UNNERuntimeIREECpu::GetModelDataIdentifier(FString FileType, TConstArray
 bool UNNERuntimeIREECpu::CanCreateModelCPU(TObjectPtr<UNNEModelData> ModelData) const { return false; };
 TSharedPtr<UE::NNE::IModelCPU> UNNERuntimeIREECpu::CreateModelCPU(TObjectPtr<UNNEModelData> ModelData) { return TSharedPtr<UE::NNE::IModelCPU>(); };
 
-FString UNNERuntimeIREECpu::GetIntermediateModelDirPath(const FString& PlatformName, const FString& FileIdString) { return ""; }
-FString UNNERuntimeIREECpu::GetCookedModelDirPath(const FString& PlatformName) { return ""; }
-FString UNNERuntimeIREECpu::GetPackagedModelDirPath(const FString& PlatformName) { return ""; }
-
 void UNNERuntimeIREECpu::GetUpdatedPlatformConfig(const FString& PlatformName, FConfigFile& ConfigFile, FString& ConfigFilePath) { }
 
 #endif // WITH_NNE_RUNTIME_IREE
\ No newline at end of file
diff --git a/Engine/Plugins/Experimental/NNERuntimeIREE/Source/NNERuntimeIREE/Private/NNERuntimeIREECpu.h b/Engine/Plugins/Experimental/NNERuntimeIREE/Source/NNERuntimeIREE/Private/NNERuntimeIREECpu.h
--- a/Engine/Plugins/Experimental/NNERuntimeIREE/Source/NNERuntimeIREE/Private/NNERuntimeIREECpu.h
+++ b/Engine/Plugins/Experimental/NNERuntimeIREE/Source/NNERuntimeIREE/Private/NNERuntimeIREECpu.h
@@ -5,9 +5,9 @@
 #include "CoreMinimal.h"
 #include "Misc/ConfigCacheIni.h"
 #include "Misc/Guid.h"
-#include "NNEModelData.h"
 #include "NNERuntime.h"
 #include "NNERuntimeCPU.h"
+#include "NNERuntimeIREECpuCompiler.h"
 
 #include "NNERuntimeIREECpu.generated.h"
 
@@ -19,7 +19,6 @@ class UNNERuntimeIREECpu : public UObject, public INNERuntime, public INNERuntim
 public:
 	static FGuid GUID;
 	static int32 Version;
-	static uint32 MemoryAlignment;
 
 	UNNERuntimeIREECpu() {};
 	virtual ~UNNERuntimeIREECpu() = default;
@@ -36,9 +35,5 @@ public:
 	virtual TSharedPtr<UE::NNE::IModelCPU> CreateModelCPU(TObjectPtr<UNNEModelData> ModelData) override;
 	//~ End INNERuntimeCPU Interface
 
-	static FString GetIntermediateModelDirPath(const FString& PlatformName, const FString& FileIdString);
-	static FString GetCookedModelDirPath(const FString& PlatformName);
-	static FString GetPackagedModelDirPath(const FString& PlatformName);
-
 	static void GetUpdatedPlatformConfig(const FString& PlatformName, FConfigFile& ConfigFile, FString& ConfigFilePath);
 };
\ No newline at end of file
diff --git a/Engine/Plugins/Experimental/NNERuntimeIREE/Source/NNERuntimeIREE/Private/NNERuntimeIREECpuCompiler.cpp b/Engine/Plugins/Experimental/NNERuntimeIREE/Source/NNERuntimeIREE/Private/NNERuntimeIREECpuCompiler.cpp
new file mode 100644
--- /dev/null
+++ b/Engine/Plugins/Experimental/NNERuntimeIREE/Source/NNERuntimeIREE/Private/NNERuntimeIREECpuCompiler.cpp
@@ -0,0 +1,310 @@
+// Copyright Epic Games, Inc. All Rights Reserved.
+
+#include "NNERuntimeIREECpuCompiler.h"
+
+#ifdef WITH_NNE_RUNTIME_IREE
+#if WITH_EDITOR
+
+#include "GenericPlatform/GenericPlatformFile.h"
+#include "GenericPlatform/GenericPlatformProcess.h"
+#include "HAL/FileManager.h"
+#include "HAL/PlatformFileManager.h"
+#include "Interfaces/IPluginManager.h"
+#include "Kismet/GameplayStatics.h"
+#include "Misc/FileHelper.h"
+#include "Misc/Paths.h"
+#include "Serialization/JsonReader.h"
+#include "Serialization/JsonSerializer.h"
+#include "Templates/SharedPointer.h"
+
+namespace UE::NNERuntimeIREECpu::Private
+{
+	bool ResolveEnvironmentVariables(FString& String)
+	{
+		const FString StartString = "$ENV{";
+		const FString EndString = "}";
+		FString ResultString = String;
+		int32 StartIndex = ResultString.Find(StartString, ESearchCase::CaseSensitive);
+		while (StartIndex != INDEX_NONE)
+		{
+			StartIndex += StartString.Len();
+			int32 EndIndex = ResultString.Find(EndString, ESearchCase::CaseSensitive, ESearchDir::FromStart, StartIndex);
+			if (EndIndex > StartIndex)
+			{
+				FString EnvironmentVariableName = ResultString.Mid(StartIndex, EndIndex - StartIndex);
+				FString EnvironmentVariableValue = FNNERuntimeIREECpuCompiler::GetEnvVar(EnvironmentVariableName);
+				if (EnvironmentVariableValue.IsEmpty())
+				{
+					return false;
+				}
+				else
+				{
+					ResultString.ReplaceInline(*(StartString + EnvironmentVariableName + EndString), *EnvironmentVariableValue, ESearchCase::CaseSensitive);
+				}
+			}
+			else
+			{
+				return false;
+			}
+			StartIndex = ResultString.Find(StartString, ESearchCase::CaseSensitive);
+		}
+		String = ResultString;
+		return true;
+	}
+
+	FString GetSharedLibraryEntryPointName(const FString& HeaderString)
+	{
+		FString SearchString = "iree_hal_executable_library_header_t**";
+		int32 Start = HeaderString.Find(SearchString);
+		if (Start == INDEX_NONE)
+		{
+			return "";
+		}
+		Start += SearchString.Len();
+		int32 End = HeaderString.Find("(", ESearchCase::CaseSensitive, ESearchDir::FromStart, Start);
+		if (End <= Start)
+		{
+			return "";
+		}
+		return HeaderString.Mid(Start, End - Start).TrimStartAndEnd();
+	}
+
+	void RunCommand(const FString& Command, const FString& Arguments)
+	{
+		void* PipeRead = nullptr;
+		void* PipeWrite = nullptr;
+		FPlatformProcess::CreatePipe(PipeRead, PipeWrite);
+		FProcHandle ProcHandle = FPlatformProcess::CreateProc(*Command, *Arguments, false, true, true, nullptr, 0, nullptr, PipeWrite, PipeRead);
+		FPlatformProcess::WaitForProc(ProcHandle);
+		FPlatformProcess::CloseProc(ProcHandle);
+		FPlatformProcess::ClosePipe(PipeRead, PipeWrite);
+	}
+
+	FNNERuntimeIREECpuCompiler::FNNERuntimeIREECpuCompiler(const FString& InCompilerCommand, const FString& InLinkerCommand, const FString& InSharedLibExt, TConstArrayView<FIREEBuildTarget> InTargets) : CompilerCommand(InCompilerCommand), LinkerCommand(InLinkerCommand), SharedLibExt(InSharedLibExt), Targets(InTargets)
+	{
+
+	}
+
+	TUniquePtr<FNNERuntimeIREECpuCompiler> FNNERuntimeIREECpuCompiler::Make(const FString& InTargetPlatformName)
+	{
+		using namespace Private;
+
+		FString PluginDir = FPaths::ConvertRelativePathToFull(*IPluginManager::Get().FindPlugin(UE_PLUGIN_NAME)->GetBaseDir());
+		FString BuildConfigFileName = FString("IREE_") + UGameplayStatics::GetPlatformName() + "_To_" + InTargetPlatformName + ".json";
+		TArray<FString> BuildConfigFilePaths =
+		{
+			FPaths::Combine(FPaths::ConvertRelativePathToFull(*FPaths::ProjectConfigDir()), BuildConfigFileName),
+			FPaths::Combine(PluginDir, "Config", BuildConfigFileName),
+			FPaths::Combine(FPaths::ConvertRelativePathToFull(*FPaths::EngineDir()), "Platforms", InTargetPlatformName, "Plugins", UE_PLUGIN_NAME, "Config", BuildConfigFileName),
+			FPaths::Combine(FPaths::ConvertRelativePathToFull(*FPaths::EngineDir()), "Platforms", InTargetPlatformName, "Plugins", "Experimental", UE_PLUGIN_NAME, "Config", BuildConfigFileName)
+		};
+
+		FString CompilerCommand;
+		FString LinkerCommand;
+		FString SharedLibExt;
+		TArray<FIREEBuildTarget> Targets;
+		IPlatformFile& PlatformFile = FPlatformFileManager::Get().GetPlatformFile();
+		for (FString BuildConfigFilePath : BuildConfigFilePaths)
+		{
+			if (PlatformFile.FileExists(*BuildConfigFilePath))
+			{
+				FString BuildConfigFileString;
+				if (FFileHelper::LoadFileToString(BuildConfigFileString, *BuildConfigFilePath))
+				{
+					TSharedRef<TJsonReader<TCHAR>> JsonReader = TJsonReaderFactory<TCHAR>::Create(BuildConfigFileString);
+					TSharedPtr<FJsonObject> JsonObject = MakeShareable(new FJsonObject);
+					FIREEBuildConfig BuildConfig;
+					if (FJsonSerializer::Deserialize(JsonReader, JsonObject) && JsonObject.IsValid() && BuildConfig.FromJson(JsonObject))
+					{
+						if (BuildConfig.Targets.IsEmpty())
+						{
+							UE_LOG(LogTemp, Warning, TEXT("UNNERuntimeIREECpu could not find targets in %s"), *BuildConfigFilePath);
+							continue;
+						}
+
+						FString TmpCompilerCommand;
+						for (int32 i = 0; i < BuildConfig.CompilerCommand.Num(); i++)
+						{
+							if (ResolveEnvironmentVariables(BuildConfig.CompilerCommand[i]))
+							{
+								BuildConfig.CompilerCommand[i].ReplaceInline(*FString("${PLUGIN_DIR}"), *PluginDir);
+								BuildConfig.CompilerCommand[i].ReplaceInline(*FString("${PROJECT_DIR}"), *FPaths::ProjectDir());
+								if (PlatformFile.FileExists(*BuildConfig.CompilerCommand[i]))
+								{
+									TmpCompilerCommand = BuildConfig.CompilerCommand[i];
+									break;
+								}
+							}
+							else
+							{
+								UE_LOG(LogTemp, Warning, TEXT("UNNERuntimeIREECpu could not replace environment variables in %s"), *BuildConfig.CompilerCommand[i]);
+							}
+						}
+						if (TmpCompilerCommand.IsEmpty())
+						{
+							UE_LOG(LogTemp, Warning, TEXT("UNNERuntimeIREECpu could not find the compiler executable in %s"), *BuildConfigFilePath);
+							continue;
+						}
+
+						FString TmpLinkerCommand;
+						for (int32 i = 0; i < BuildConfig.LinkerCommand.Num(); i++)
+						{
+							if (ResolveEnvironmentVariables(BuildConfig.LinkerCommand[i]))
+							{
+								BuildConfig.LinkerCommand[i].ReplaceInline(*FString("${PLUGIN_DIR}"), *PluginDir);
+								BuildConfig.LinkerCommand[i].ReplaceInline(*FString("${PROJECT_DIR}"), *FPaths::ProjectDir());
+								if (PlatformFile.FileExists(*BuildConfig.LinkerCommand[i]))
+								{
+									TmpLinkerCommand = BuildConfig.LinkerCommand[i];
+									break;
+								}
+							}
+							else
+							{
+								UE_LOG(LogTemp, Warning, TEXT("UNNERuntimeIREECpu could not replace environment variables in %s"), *BuildConfig.LinkerCommand[i]);
+							}
+						}
+						if (TmpLinkerCommand.IsEmpty())
+						{
+							UE_LOG(LogTemp, Warning, TEXT("UNNERuntimeIREECpu could not find the linker executable in %s"), *BuildConfigFilePath);
+							continue;
+						}
+
+						CompilerCommand = TmpCompilerCommand;
+						LinkerCommand = TmpLinkerCommand;
+						SharedLibExt = BuildConfig.SharedLibExt;
+						Targets = BuildConfig.Targets;
+						break;
+					}
+					else
+					{
+						UE_LOG(LogTemp, Warning, TEXT("UNNERuntimeIREECpu could not parse build config file %s"), *BuildConfigFilePath);
+					}
+				}
+				else
+				{
+					UE_LOG(LogTemp, Warning, TEXT("UNNERuntimeIREECpu could not read build config file %s"), *BuildConfigFilePath);
+				}
+			}
+		}
+		if (CompilerCommand.IsEmpty() || LinkerCommand.IsEmpty() || Targets.IsEmpty())
+		{
+			return TUniquePtr<FNNERuntimeIREECpuCompiler>();
+		}
+		return MakeUnique<FNNERuntimeIREECpuCompiler>(CompilerCommand, LinkerCommand, SharedLibExt, Targets);
+	}
+
+	bool FNNERuntimeIREECpuCompiler::CompileMlir(TConstArrayView<uint8> InFileData, const FString& InModelName, const FString& InIntermediateDir, const FString& InStagingDir, TArray<FIREECompilerResult>& OutCompilerResults)
+	{
+		using namespace Private;
+
+		FString InputFilePath = FPaths::Combine(InIntermediateDir, InModelName) + ".mlir";
+		FFileHelper::SaveArrayToFile(InFileData, *InputFilePath);
+
+		TArray<FIREECompilerResult> Results;
+		bool bResult = true;
+		IPlatformFile& PlatformFile = FPlatformFileManager::Get().GetPlatformFile();
+		for (int32 i = 0; i < Targets.Num(); i++)
+		{
+			FString IntermediateFilePathNoExt = FPaths::Combine(InIntermediateDir, Targets[i].Architecture, InModelName);
+			FString ObjectFilePath = IntermediateFilePathNoExt + ".o";
+			FString VmfbFilePath = IntermediateFilePathNoExt + ".vmfb";
+			FString SharedLibFilePath = IntermediateFilePathNoExt + SharedLibExt;
+
+			FString CompilerArguments = Targets[i].CompilerArguments;
+			if (!ResolveEnvironmentVariables(CompilerArguments))
+			{
+				UE_LOG(LogTemp, Warning, TEXT("UNNERuntimeIREECpu could not replace environment variables in %s"), *Targets[i].CompilerArguments);
+				continue;
+			}
+			CompilerArguments.ReplaceInline(*FString("${OBJECT_PATH}"), *(FString("\"") + ObjectFilePath + "\""));
+			CompilerArguments.ReplaceInline(*FString("${VMFB_PATH}"), *(FString("\"") + VmfbFilePath + "\""));
+			CompilerArguments.ReplaceInline(*FString("${INPUT_PATH}"), *(FString("\"") + InputFilePath + "\""));
+
+			RunCommand(CompilerCommand, CompilerArguments);
+
+			if (!PlatformFile.FileExists(*ObjectFilePath) || !PlatformFile.FileExists(*VmfbFilePath))
+			{
+				UE_LOG(LogTemp, Warning, TEXT("UNNERuntimeIREECpu failed to compile the model \"%s\" using the command:"), *InputFilePath);
+				UE_LOG(LogTemp, Warning, TEXT("\"%s\" %s"), *CompilerCommand, *CompilerArguments);
+				bResult = false;
+				continue;
+			}
+
+			FString LinkerArguments = Targets[i].LinkerArguments;
+			if (!ResolveEnvironmentVariables(LinkerArguments))
+			{
+				UE_LOG(LogTemp, Warning, TEXT("UNNERuntimeIREECpu could not replace environment variables in %s"), *Targets[i].LinkerArguments);
+				bResult = false;
+				continue;
+			}
+			LinkerArguments.ReplaceInline(*FString("${OBJECT_PATH}"), *(FString("\"") + ObjectFilePath + "\""));
+			LinkerArguments.ReplaceInline(*FString("${SHARED_LIB_PATH}"), *(FString("\"") + SharedLibFilePath + "\""));
+
+			RunCommand(LinkerCommand, LinkerArguments);
+
+			if (!PlatformFile.FileExists(*SharedLibFilePath))
+			{
+				UE_LOG(LogTemp, Warning, TEXT("UNNERuntimeIREECpu failed to link the model \"%s\" using the command:"), *InputFilePath);
+				UE_LOG(LogTemp, Warning, TEXT("\"%s\" %s"), *LinkerCommand, *LinkerArguments);
+				bResult = false;
+				continue;
+			}
+
+			FString StagingFilePathNoExt = FPaths::Combine(InStagingDir, Targets[i].Architecture, InModelName);
+			FString StagedVmfbFilePath = StagingFilePathNoExt + ".vmfb";
+			FString StagedSharedLibFilePath = StagingFilePathNoExt + SharedLibExt;
+			if (IFileManager::Get().Copy(*StagedSharedLibFilePath, *SharedLibFilePath) != COPY_OK)
+			{
+				UE_LOG(LogTemp, Warning, TEXT("UNNERuntimeIREECpu failed to copy \"%s\" to \"%s\""), *SharedLibFilePath, *StagedSharedLibFilePath);
+				bResult = false;
+			}
+			if (IFileManager::Get().Copy(*StagedVmfbFilePath, *VmfbFilePath) != COPY_OK)
+			{
+				UE_LOG(LogTemp, Warning, TEXT("UNNERuntimeIREECpu failed to copy \"%s\" to \"%s\""), *VmfbFilePath, *StagedVmfbFilePath);
+				bResult = false;
+			}
+
+			FString SharedLibraryEntryPointName = "";
+			FString HeaderPath = IntermediateFilePathNoExt + ".h";
+			if (!PlatformFile.FileExists(*HeaderPath))
+			{
+				UE_LOG(LogTemp, Warning, TEXT("UNNERuntimeIREECpu could not find the model header \"%s\""), *HeaderPath);
+				bResult = false;
+				continue;
+			}
+			FString HeaderString;
+			if (!FFileHelper::LoadFileToString(HeaderString, *HeaderPath) || HeaderString.IsEmpty())
+			{
+				UE_LOG(LogTemp, Warning, TEXT("UNNERuntimeIREECpu could not read the model header \"%s\""), *HeaderPath);
+				bResult = false;
+				continue;
+			}
+			SharedLibraryEntryPointName = GetSharedLibraryEntryPointName(HeaderString);
+			if (SharedLibraryEntryPointName.IsEmpty())
+			{
+				UE_LOG(LogTemp, Warning, TEXT("UNNERuntimeIREECpu could not find the entry point in model header \"%s\""), *HeaderPath);
+				bResult = false;
+				continue;
+			}
+
+			FIREECompilerResult Result;
+			Result.Architecture = Targets[i].Architecture;
+			Result.RelativeDirPath = Targets[i].Architecture;
+			Result.SharedLibraryFileName = InModelName + SharedLibExt;
+			Result.VmfbFileName = InModelName + ".vmfb";
+			Result.SharedLibraryEntryPointName = SharedLibraryEntryPointName;
+			Results.Add(Result);
+		}
+
+		bResult &= !Results.IsEmpty();
+		if (bResult)
+		{
+			OutCompilerResults = Results;
+		}
+		return bResult;
+	}
+} // UE::NNERuntimeIREECpu::Private
+
+#endif // WITH_EDITOR
+#endif // WITH_NNE_RUNTIME_IREE
\ No newline at end of file
diff --git a/Engine/Plugins/Experimental/NNERuntimeIREE/Source/NNERuntimeIREE/Private/NNERuntimeIREECpuCompiler.h b/Engine/Plugins/Experimental/NNERuntimeIREE/Source/NNERuntimeIREE/Private/NNERuntimeIREECpuCompiler.h
new file mode 100644
--- /dev/null
+++ b/Engine/Plugins/Experimental/NNERuntimeIREE/Source/NNERuntimeIREE/Private/NNERuntimeIREECpuCompiler.h
@@ -0,0 +1,76 @@
+// Copyright Epic Games, Inc. All Rights Reserved.
+
+#pragma once
+
+#ifdef WITH_NNE_RUNTIME_IREE
+#if WITH_EDITOR
+
+#include "Containers/Array.h"
+#include "Containers/ArrayView.h"
+#include "Containers/UnrealString.h"
+#include "GenericPlatform/GenericPlatformMisc.h"
+#include "Serialization/JsonSerializerMacros.h"
+#include "Templates/UniquePtr.h"
+
+namespace UE::NNERuntimeIREECpu::Private
+{
+	struct FIREEBuildTarget : public FJsonSerializable
+	{
+		FString Architecture;
+		FString CompilerArguments;
+		FString LinkerArguments;
+
+		BEGIN_JSON_SERIALIZER
+			JSON_SERIALIZE("Architecture", Architecture);
+			JSON_SERIALIZE("CompilerArguments", CompilerArguments);
+			JSON_SERIALIZE("LinkerArguments", LinkerArguments);
+		END_JSON_SERIALIZER
+	};
+
+	struct FIREEBuildConfig : public FJsonSerializable
+	{
+		TArray<FString> CompilerCommand;
+		TArray<FString> LinkerCommand;
+		FString SharedLibExt;
+		TArray<FIREEBuildTarget> Targets;
+
+		BEGIN_JSON_SERIALIZER
+			JSON_SERIALIZE_ARRAY("CompilerCommand", CompilerCommand);
+			JSON_SERIALIZE_ARRAY("LinkerCommand", LinkerCommand);
+			JSON_SERIALIZE("SharedLibExt", SharedLibExt);
+			JSON_SERIALIZE_ARRAY_SERIALIZABLE("Targets", Targets, FIREEBuildTarget);
+		END_JSON_SERIALIZER
+	};
+
+	struct FIREECompilerResult
+	{
+		FString Architecture;
+		FString RelativeDirPath;
+		FString SharedLibraryFileName;
+		FString VmfbFileName;
+		FString SharedLibraryEntryPointName;
+	};
+
+	class FNNERuntimeIREECpuCompiler
+	{
+	public:
+		FNNERuntimeIREECpuCompiler(const FString& InCompilerCommand, const FString& InLinkerCommand, const FString& InSharedLibExt, TConstArrayView<FIREEBuildTarget> InTargets);
+		~FNNERuntimeIREECpuCompiler() {};
+		static TUniquePtr<FNNERuntimeIREECpuCompiler> Make(const FString& InTargetPlatformName);
+
+	public:
+		bool CompileMlir(TConstArrayView<uint8> InFileData, const FString& InModelName, const FString& InIntermediateDir, const FString& InStagingDir, TArray<FIREECompilerResult>& OutCompilerResults);
+
+		// Needs to stay in header to prevent system macros to overwrite GetEnvironmentVariable
+		static FString GetEnvVar(const FString& EnvironmentVariableName) { return FPlatformMisc::GetEnvironmentVariable(*EnvironmentVariableName); }
+
+	private:
+		FString CompilerCommand;
+		FString LinkerCommand;
+		FString SharedLibExt;
+		TArray<FIREEBuildTarget> Targets;
+	};
+} // UE::NNERuntimeIREECpu::Private
+
+#endif // WITH_EDITOR
+#endif // WITH_NNE_RUNTIME_IREE
\ No newline at end of file
diff --git a/Engine/Plugins/Experimental/NNERuntimeIREE/Source/NNERuntimeIREE/Private/NNERuntimeIREECpuModel.cpp b/Engine/Plugins/Experimental/NNERuntimeIREE/Source/NNERuntimeIREE/Private/NNERuntimeIREECpuModel.cpp
--- a/Engine/Plugins/Experimental/NNERuntimeIREE/Source/NNERuntimeIREE/Private/NNERuntimeIREECpuModel.cpp
+++ b/Engine/Plugins/Experimental/NNERuntimeIREE/Source/NNERuntimeIREE/Private/NNERuntimeIREECpuModel.cpp
@@ -57,21 +57,20 @@ namespace UE::NNERuntimeIREECpu
 		return Session->RunSyncCPU(InInputBindings, InOutputBindings);
 	}
 
-	bool FModel::Init(TSharedPtr<UE::NNE::FSharedModelData> SharedModelData, uint32 VmfbDataOffset, UE::NNERuntimeIREE::FModuleMetaData ModuleMetaData, const FString& LibraryPath, const FString& LibraryName, const FString& LibraryQueryFunctionName)
+	bool FModel::Init(const FString& DirPath, const FString& SharedLibraryFileName, const FString& VmfbFileName, const FString& LibraryQueryFunctionName, const UE::NNERuntimeIREE::FModuleMetaData& ModuleMetaData)
 	{
 		check(!Module.IsValid());
-		check(SharedModelData.IsValid());
-		check(SharedModelData->GetView().Num() > 0);
-		check(LibraryName.Len() > 0);
-		check(LibraryQueryFunctionName.Len() > 0);
+		check(!SharedLibraryFileName.IsEmpty());
+		check(!VmfbFileName.IsEmpty());
+		check(!LibraryQueryFunctionName.IsEmpty());
 
-		Module = UE::NNERuntimeIREE::FIREEModule::MakeModule(SharedModelData, VmfbDataOffset, ModuleMetaData);
+		Module = UE::NNERuntimeIREE::FIREEModule::MakeModule(DirPath, VmfbFileName, ModuleMetaData);
 		if (!Module.IsValid())
 		{
 			return false;
 		}
 
-		TSharedPtr<UE::NNERuntimeIREE::FIREELibrary> Library = UE::NNERuntimeIREE::FIREELibrary::MakeLibrary(LibraryPath, LibraryName);
+		TSharedPtr<UE::NNERuntimeIREE::FIREELibrary> Library = UE::NNERuntimeIREE::FIREELibrary::MakeLibrary(DirPath, SharedLibraryFileName);
 		if (!Library.IsValid())
 		{
 			Module.Reset();
diff --git a/Engine/Plugins/Experimental/NNERuntimeIREE/Source/NNERuntimeIREE/Private/NNERuntimeIREECpuModel.h b/Engine/Plugins/Experimental/NNERuntimeIREE/Source/NNERuntimeIREE/Private/NNERuntimeIREECpuModel.h
--- a/Engine/Plugins/Experimental/NNERuntimeIREE/Source/NNERuntimeIREE/Private/NNERuntimeIREECpuModel.h
+++ b/Engine/Plugins/Experimental/NNERuntimeIREE/Source/NNERuntimeIREE/Private/NNERuntimeIREECpuModel.h
@@ -4,11 +4,9 @@
 
 #ifdef WITH_NNE_RUNTIME_IREE
 
-#include "CoreMinimal.h"
-#include "NNEModelData.h"
 #include "NNERuntimeCPU.h"
 #include "NNERuntimeIREECommon.h"
-#include "NNETypes.h"
+#include "NNERuntimeIREEMetaData.h"
 
 namespace UE::NNERuntimeIREECpu
 {
@@ -38,7 +36,7 @@ namespace UE::NNERuntimeIREECpu
 		FModel() {};
 		virtual ~FModel() = default;
 
-		bool Init(TSharedPtr<UE::NNE::FSharedModelData> SharedModelData, uint32 VmfbDataOffset, UE::NNERuntimeIREE::FModuleMetaData ModuleMetaData, const FString& LibraryPath, const FString& LibraryName, const FString& LibraryQueryFunctionName);
+		bool Init(const FString& DirPath, const FString& SharedLibraryFileName, const FString& VmfbFileName, const FString& LibraryQueryFunctionName, const UE::NNERuntimeIREE::FModuleMetaData& ModuleMetaData);
 
 		//~ Begin IModelCPU Interface
 		virtual TSharedPtr<UE::NNE::IModelInstanceCPU> CreateModelInstanceCPU() override;
diff --git a/Engine/Plugins/Experimental/NNERuntimeIREE/Source/NNERuntimeIREE/Private/NNERuntimeIREEGpu.cpp b/Engine/Plugins/Experimental/NNERuntimeIREE/Source/NNERuntimeIREE/Private/NNERuntimeIREEGpu.cpp
--- a/Engine/Plugins/Experimental/NNERuntimeIREE/Source/NNERuntimeIREE/Private/NNERuntimeIREEGpu.cpp
+++ b/Engine/Plugins/Experimental/NNERuntimeIREE/Source/NNERuntimeIREE/Private/NNERuntimeIREEGpu.cpp
@@ -3,6 +3,7 @@
 #include "NNERuntimeIREEGpu.h"
 
 #include "EngineAnalytics.h"
+#include "Interfaces/ITargetPlatform.h"
 #include "Kismet/GameplayStatics.h"
 #include "NNERuntimeIREECommon.h"
 
@@ -45,8 +46,8 @@ TSharedPtr<UE::NNE::FSharedModelData> UNNERuntimeIREEGpu::CreateModelData(FStrin
 
 FString UNNERuntimeIREEGpu::GetModelDataIdentifier(FString FileType, TConstArrayView<uint8> FileData, FGuid FileId, const ITargetPlatform* TargetPlatform)
 {
-	FString PlatformDisplayName = UE::NNERuntimeIREE::GetTargetPlatformDisplayName(TargetPlatform);
-	return FileId.ToString(EGuidFormats::Digits) + "-" + GetRuntimeName() + "-" + GetGUID().ToString(EGuidFormats::Digits) + "-" + FString::FromInt(GetVersion()) + "-" + PlatformDisplayName;
+	FString PlatformName = TargetPlatform ? TargetPlatform->IniPlatformName() : UGameplayStatics::GetPlatformName();
+	return GetRuntimeName() + "-" + GetGUID().ToString(EGuidFormats::Digits) + "-" + FString::FromInt(GetVersion()) + "-" + FileId.ToString(EGuidFormats::Digits) + "-" + PlatformName;
 }
 
 bool UNNERuntimeIREEGpu::CanCreateModelGPU(TObjectPtr<UNNEModelData> ModelData) const
diff --git a/Engine/Plugins/Experimental/NNERuntimeIREE/Source/NNERuntimeIREE/Private/NNERuntimeIREEModule.cpp b/Engine/Plugins/Experimental/NNERuntimeIREE/Source/NNERuntimeIREE/Private/NNERuntimeIREEModule.cpp
--- a/Engine/Plugins/Experimental/NNERuntimeIREE/Source/NNERuntimeIREE/Private/NNERuntimeIREEModule.cpp
+++ b/Engine/Plugins/Experimental/NNERuntimeIREE/Source/NNERuntimeIREE/Private/NNERuntimeIREEModule.cpp
@@ -2,16 +2,17 @@
 
 #include "NNERuntimeIREEModule.h"
 
+#include "Modules/ModuleManager.h"
+
 #ifdef WITH_NNE_RUNTIME_IREE
 
+#include "CoreMinimal.h"
 #include "NNE.h"
-#include "NNERuntime.h"
 
-#ifdef WITH_EDITOR
-#include "Serialization/Archive.h"
+#if WITH_EDITOR
 #include "HAL/FileManager.h"
-#include "Interfaces/ITargetPlatformManagerModule.h"
 #include "Interfaces/ITargetPlatform.h"
+#include "Interfaces/ITargetPlatformManagerModule.h"
 #include "Misc/ConfigCacheIni.h"
 #endif // WITH_EDITOR
 
@@ -27,7 +28,7 @@ void FNNERuntimeIREEModule::StartupModule()
 		UE::NNE::RegisterRuntime(NNERuntimeIREECpu.Get());
 	}
 
-#ifdef WITH_EDITOR
+#if WITH_EDITOR
 	ITargetPlatformManagerModule* TargetPlatformManagerModule = GetTargetPlatformManager();
 	if (TargetPlatformManagerModule)
 	{
diff --git a/Engine/Plugins/Experimental/NNERuntimeIREE/Source/NNERuntimeIREE/Private/NNERuntimeIREEModule.h b/Engine/Plugins/Experimental/NNERuntimeIREE/Source/NNERuntimeIREE/Private/NNERuntimeIREEModule.h
--- a/Engine/Plugins/Experimental/NNERuntimeIREE/Source/NNERuntimeIREE/Private/NNERuntimeIREEModule.h
+++ b/Engine/Plugins/Experimental/NNERuntimeIREE/Source/NNERuntimeIREE/Private/NNERuntimeIREEModule.h
@@ -2,13 +2,13 @@
 
 #pragma once
 
-#include "CoreMinimal.h"
-#include "Modules/ModuleManager.h"
+#include "Modules/ModuleInterface.h"
 
 #ifdef WITH_NNE_RUNTIME_IREE
 #include "NNERuntimeIREECpu.h"
 #include "NNERuntimeIREEGpu.h"
 #include "NNERuntimeIREERdg.h"
+#include "UObject/WeakObjectPtrTemplates.h"
 #endif // WITH_NNE_RUNTIME_IREE
 
 class FNNERuntimeIREEModule : public IModuleInterface
diff --git a/Engine/Plugins/Experimental/NNERuntimeIREE/Source/NNERuntimeIREE/Private/NNERuntimeIREERdg.cpp b/Engine/Plugins/Experimental/NNERuntimeIREE/Source/NNERuntimeIREE/Private/NNERuntimeIREERdg.cpp
--- a/Engine/Plugins/Experimental/NNERuntimeIREE/Source/NNERuntimeIREE/Private/NNERuntimeIREERdg.cpp
+++ b/Engine/Plugins/Experimental/NNERuntimeIREE/Source/NNERuntimeIREE/Private/NNERuntimeIREERdg.cpp
@@ -3,6 +3,7 @@
 #include "NNERuntimeIREERdg.h"
 
 #include "EngineAnalytics.h"
+#include "Interfaces/ITargetPlatform.h"
 #include "Kismet/GameplayStatics.h"
 #include "NNERuntimeIREECommon.h"
 
@@ -32,8 +33,8 @@ TSharedPtr<UE::NNE::FSharedModelData> UNNERuntimeIREERdg::CreateModelData(FStrin
 
 FString UNNERuntimeIREERdg::GetModelDataIdentifier(FString FileType, TConstArrayView<uint8> FileData, FGuid FileId, const ITargetPlatform* TargetPlatform)
 {
-	FString PlatformDisplayName = UE::NNERuntimeIREE::GetTargetPlatformDisplayName(TargetPlatform);
-	return FileId.ToString(EGuidFormats::Digits) + "-" + GUID.ToString(EGuidFormats::Digits) + "-" + FString::FromInt(Version) + "-" + PlatformDisplayName;
+	FString PlatformName = TargetPlatform ? TargetPlatform->IniPlatformName() : UGameplayStatics::GetPlatformName();
+	return GetRuntimeName() + "-" + GUID.ToString(EGuidFormats::Digits) + "-" + FString::FromInt(UNNERuntimeIREERdg::Version) + "-" + FileId.ToString(EGuidFormats::Digits) + "-" + PlatformName;
 }
 
 bool UNNERuntimeIREERdg::CanCreateModelRDG(TObjectPtr<UNNEModelData> ModelData) const
