[NNE][NNERuntimeIREE] Fixing issues with IREE build chain on Mac

Problem: Wrong linker flags, missing folder creation and wrong architecture selection prevented IREE from creating and loading models.

Solution: Fix the individual issues

Tested: By running a compiled model locally on Mac

#rnx

diff --git a/Engine/Plugins/Experimental/NNERuntimeIREE/Config/IREE_Mac_To_Mac.json b/Engine/Plugins/Experimental/NNERuntimeIREE/Config/IREE_Mac_To_Mac.json
--- a/Engine/Plugins/Experimental/NNERuntimeIREE/Config/IREE_Mac_To_Mac.json
+++ b/Engine/Plugins/Experimental/NNERuntimeIREE/Config/IREE_Mac_To_Mac.json
@@ -6,12 +6,12 @@
 		{
 			"Architecture": "x86_64",
 			"CompilerArguments": "--iree-llvmcpu-target-triple=x86_64-apple-darwin --iree-input-type=auto --iree-hal-target-backends=llvm-cpu --iree-llvmcpu-link-embedded=false --iree-llvmcpu-link-static --output-format=vm-bytecode --iree-llvmcpu-static-library-output-path=${OBJECT_PATH} -o ${VMFB_PATH} ${INPUT_PATH}",
-			"LinkerArguments": "-dylib -arch x86_64 ${SHARED_LIB_PATH} ${OBJECT_PATH}"
+			"LinkerArguments": "-dylib -platform_version macos 11.0.0 11.0.0 -arch x86_64 -o ${SHARED_LIB_PATH} ${OBJECT_PATH}"
 		},
 		{
 			"Architecture": "arm64",
 			"CompilerArguments": "--iree-llvmcpu-target-triple=arm64-apple-darwin --iree-input-type=auto --iree-hal-target-backends=llvm-cpu --iree-llvmcpu-link-embedded=false --iree-llvmcpu-link-static --output-format=vm-bytecode --iree-llvmcpu-static-library-output-path=${OBJECT_PATH} -o ${VMFB_PATH} ${INPUT_PATH}",
-			"LinkerArguments": "-dylib -arch arm64 ${SHARED_LIB_PATH} ${OBJECT_PATH}"
+			"LinkerArguments": "-dylib -platform_version macos 11.0.0 11.0.0 -arch arm64 -o ${SHARED_LIB_PATH} ${OBJECT_PATH}"
 		}
 	]
-}
+}
\ No newline at end of file
diff --git a/Engine/Plugins/Experimental/NNERuntimeIREE/Source/NNERuntimeIREE/Private/NNERuntimeIREECpu.cpp b/Engine/Plugins/Experimental/NNERuntimeIREE/Source/NNERuntimeIREE/Private/NNERuntimeIREECpu.cpp
--- a/Engine/Plugins/Experimental/NNERuntimeIREE/Source/NNERuntimeIREE/Private/NNERuntimeIREECpu.cpp
+++ b/Engine/Plugins/Experimental/NNERuntimeIREE/Source/NNERuntimeIREE/Private/NNERuntimeIREECpu.cpp
@@ -259,7 +259,7 @@ TSharedPtr<UE::NNE::IModelCPU> UNNERuntimeIREECpu::CreateModelCPU(TObjectPtr<UNN
 	}
 
 	FString CurrentArchitecture = "";
-#ifdef PLATFORM_CPU_X86_FAMILY
+#if PLATFORM_CPU_X86_FAMILY
 	CurrentArchitecture = "x86_64";
 #elif PLATFORM_CPU_ARM_FAMILY
 	CurrentArchitecture = "arm64";
diff --git a/Engine/Plugins/Experimental/NNERuntimeIREE/Source/NNERuntimeIREE/Private/NNERuntimeIREECpuCompiler.cpp b/Engine/Plugins/Experimental/NNERuntimeIREE/Source/NNERuntimeIREE/Private/NNERuntimeIREECpuCompiler.cpp
--- a/Engine/Plugins/Experimental/NNERuntimeIREE/Source/NNERuntimeIREE/Private/NNERuntimeIREECpuCompiler.cpp
+++ b/Engine/Plugins/Experimental/NNERuntimeIREE/Source/NNERuntimeIREE/Private/NNERuntimeIREECpuCompiler.cpp
@@ -206,7 +206,9 @@ namespace UE::NNERuntimeIREECpu::Private
 		IPlatformFile& PlatformFile = FPlatformFileManager::Get().GetPlatformFile();
 		for (int32 i = 0; i < Targets.Num(); i++)
 		{
-			FString IntermediateFilePathNoExt = FPaths::Combine(InIntermediateDir, Targets[i].Architecture, InModelName);
+			FString IntermediateDirPath = FPaths::Combine(InIntermediateDir, Targets[i].Architecture);
+			PlatformFile.CreateDirectoryTree(*IntermediateDirPath);
+			FString IntermediateFilePathNoExt = FPaths::Combine(IntermediateDirPath, InModelName);
 			FString ObjectFilePath = IntermediateFilePathNoExt + ".o";
 			FString VmfbFilePath = IntermediateFilePathNoExt + ".vmfb";
 			FString SharedLibFilePath = IntermediateFilePathNoExt + SharedLibExt;
