[NNERuntimeIREE]
* Fixed unity non-pch compile error

diff --git a/Engine/Plugins/Experimental/NNERuntimeIREE/Source/NNERuntimeIREE/Private/NNERuntimeIREECommon.cpp b/Engine/Plugins/Experimental/NNERuntimeIREE/Source/NNERuntimeIREE/Private/NNERuntimeIREECommon.cpp
--- a/Engine/Plugins/Experimental/NNERuntimeIREE/Source/NNERuntimeIREE/Private/NNERuntimeIREECommon.cpp
+++ b/Engine/Plugins/Experimental/NNERuntimeIREE/Source/NNERuntimeIREE/Private/NNERuntimeIREECommon.cpp
@@ -8,6 +8,7 @@
 #include "iree/hal/local/loaders/static_library_loader.h"
 #include "iree/modules/hal/types.h"
 #include "iree/vm/bytecode/module.h"
+#undef GetObject // This is needed for unity nopch builds to not throw an error  (one of the above headers include windows.h which has #define GetObject to GetObjectW)
 
 namespace UE::NNERuntimeIREE
 {
