[NNE][NNERuntimeIREE] Moving meta data parsing to import to allow main function selection for models containing multiple ones

Problem: MLIR supports models with multiple exported functions from which one needs to be picked to run inference trhough the NNE API

Solution: Parse the mlir file on import and add all functions with signatures as additional file data

Tested: local build for win64 adn consoles, presubmit preflight

#rnx

diff --git a/Engine/Plugins/Experimental/NNERuntimeIREE/Source/NNERuntimeIREE/Internal/NNERuntimeIREEMetaData.h b/Engine/Plugins/Experimental/NNERuntimeIREE/Source/NNERuntimeIREE/Internal/NNERuntimeIREEMetaData.h
new file mode 100644
--- /dev/null
+++ b/Engine/Plugins/Experimental/NNERuntimeIREE/Source/NNERuntimeIREE/Internal/NNERuntimeIREEMetaData.h
@@ -0,0 +1,34 @@
+// Copyright Epic Games, Inc. All Rights Reserved.
+
+#pragma once
+
+#include "CoreMinimal.h"
+#include "NNETypes.h"
+
+#include "NNERuntimeIREEMetaData.generated.h"
+
+namespace UE::NNERuntimeIREE
+{
+	struct FFunctionMetaData
+	{
+		FString Name;
+		TArray<UE::NNE::FTensorDesc> InputDescs;
+		TArray<UE::NNE::FTensorDesc> OutputDescs;
+	};
+} // UE::NNERuntimeIREE
+
+
+UCLASS()
+class NNERUNTIMEIREE_API UNNERuntimeIREEModuleMetaData : public UObject
+{
+	GENERATED_BODY()
+
+public:
+	// UObject interface
+	virtual void Serialize(FArchive& Ar) override;
+	// End of UObject interface
+
+	bool ParseFromString(const FString& ModuleString);
+
+	TArray<UE::NNERuntimeIREE::FFunctionMetaData> FunctionMetaData;
+};
diff --git a/Engine/Plugins/Experimental/NNERuntimeIREE/Source/NNERuntimeIREE/Private/NNERuntimeIREECommon.cpp b/Engine/Plugins/Experimental/NNERuntimeIREE/Source/NNERuntimeIREE/Private/NNERuntimeIREECommon.cpp
--- a/Engine/Plugins/Experimental/NNERuntimeIREE/Source/NNERuntimeIREE/Private/NNERuntimeIREECommon.cpp
+++ b/Engine/Plugins/Experimental/NNERuntimeIREE/Source/NNERuntimeIREE/Private/NNERuntimeIREECommon.cpp
@@ -4,90 +4,33 @@
 
 #ifdef WITH_NNE_RUNTIME_IREE
 
+#include "GenericPlatform/GenericPlatformProcess.h"
+#include "HAL/FileManager.h"
+#include "Misc/Paths.h"
+#include "Misc/ScopeLock.h"
+#include "Serialization/Archive.h"
+
+#if PLATFORM_MICROSOFT
+#include "Microsoft/AllowMicrosoftPlatformTypes.h"
+#include "Microsoft/AllowMicrosoftPlatformAtomics.h"
+#endif // PLATFORM_MICROSOFT
+THIRD_PARTY_INCLUDES_START
 #include "iree/hal/drivers/local_sync/sync_device.h"
 #include "iree/hal/local/loaders/static_library_loader.h"
-#include "iree/modules/hal/types.h"
+#include "iree/runtime/call.h"
+#include "iree/runtime/instance.h"
+#include "iree/runtime/session.h"
 #include "iree/vm/bytecode/module.h"
-#undef GetObject // This is needed for unity nopch builds to not throw an error  (one of the above headers include windows.h which has #define GetObject to GetObjectW)
+THIRD_PARTY_INCLUDES_END
+#if PLATFORM_MICROSOFT
+#include "Microsoft/HideMicrosoftPlatformAtomics.h"
+#include "Microsoft/HideMicrosoftPlatformTypes.h"
+#endif // PLATFORM_MICROSOF
 
 namespace UE::NNERuntimeIREE
 {
 	namespace Private
 	{
-		ENNETensorDataType GetTypeFromString(const FString& TypeString)
-		{
-			if (TypeString.StartsWith("char"))
-			{
-				return ENNETensorDataType::Char;
-			}
-			if (TypeString.StartsWith("bool") || TypeString.StartsWith("i1"))
-			{
-				return ENNETensorDataType::Boolean;
-			}
-			else if (TypeString.StartsWith("half"))
-			{
-				return ENNETensorDataType::Half;
-			}
-			else if (TypeString.StartsWith("f"))
-			{
-				if (TypeString.StartsWith("f16"))
-				{
-					return ENNETensorDataType::Half;
-				}
-				else if (TypeString.StartsWith("float") || TypeString.StartsWith("f32"))
-				{
-					return ENNETensorDataType::Float;
-				}
-				else if (TypeString.StartsWith("f64"))
-				{
-					return ENNETensorDataType::Double;
-				}
-			}
-			else if (TypeString.StartsWith("double"))
-			{
-				return ENNETensorDataType::Double;
-			}
-			else if (TypeString.StartsWith("i") || TypeString.StartsWith("si"))
-			{
-				if (TypeString.EndsWith("i8"))
-				{
-					return ENNETensorDataType::Int8;
-				}
-				else if (TypeString.EndsWith("i16"))
-				{
-					return ENNETensorDataType::Int16;
-				}
-				else if (TypeString.EndsWith("i32") || TypeString.EndsWith("int"))
-				{
-					return ENNETensorDataType::Int32;
-				}
-				else if (TypeString.EndsWith("i64"))
-				{
-					return ENNETensorDataType::Int64;
-				}
-			}
-			else if (TypeString.StartsWith("ui"))
-			{
-				if (TypeString.EndsWith("i8"))
-				{
-					return ENNETensorDataType::UInt8;
-				}
-				else if (TypeString.EndsWith("i16"))
-				{
-					return ENNETensorDataType::UInt16;
-				}
-				else if (TypeString.EndsWith("i32"))
-				{
-					return ENNETensorDataType::UInt32;
-				}
-				else if (TypeString.EndsWith("i64"))
-				{
-					return ENNETensorDataType::UInt64;
-				}
-			}
-			return ENNETensorDataType::None;
-		}
-
 		iree_hal_element_types_t NNEToIREEType(ENNETensorDataType Type)
 		{
 			switch (Type)
@@ -148,6 +91,17 @@ namespace UE::NNERuntimeIREE
 					break;
 			}
 		}
+
+		void PrintIREEError(const FString& Message, iree_status_t Status)
+		{
+			iree_host_size_t TrueLength = 0;
+			iree_status_format(Status, 0, (char*)nullptr, &TrueLength);
+			void* ErrorString = FMemory::Malloc(TrueLength + 1);
+			((char*)ErrorString)[TrueLength] = (char)0;
+			iree_status_format(Status, TrueLength, (char*)ErrorString, &TrueLength);
+			UE_LOG(LogTemp, Error, TEXT("%s: %s"), *Message, *FString(StringCast<TCHAR>(static_cast<const ANSICHAR*>(ErrorString)).Get()));
+			FMemory::Free(ErrorString);
+		}
 	} // Private
 
 	class FIREEInstance
@@ -193,17 +147,6 @@ namespace UE::NNERuntimeIREE
 		iree_runtime_instance_release(Instance);
 	}
 
-	void PrintIREEError(FString Message, iree_status_t Status)
-	{
-		iree_host_size_t TrueLength = 0;
-		iree_status_format(Status, 0, (char*)nullptr, &TrueLength);
-		void* ErrorString = FMemory::Malloc(TrueLength + 1);
-		((char*)ErrorString)[TrueLength] = (char)0;
-		iree_status_format(Status, TrueLength, (char*)ErrorString, &TrueLength);
-		UE_LOG(LogTemp, Error, TEXT("%s: %s"), *Message, *FString(ANSI_TO_TCHAR(ErrorString)));
-		FMemory::Free(ErrorString);
-	}
-
 	bool FIREEInstance::CreateModule(TConstArrayView<uint8> VmfbDataView, iree_vm_module_t** Module)
 	{
 		check(!VmfbDataView.IsEmpty());
@@ -218,7 +161,7 @@ namespace UE::NNERuntimeIREE
 		Status = iree_vm_bytecode_module_create(iree_runtime_instance_vm_instance(Instance), ModuleData, iree_allocator_null(), GetHostAllocator(), &TempModule);
 		if (!iree_status_is_ok(Status))
 		{
-			PrintIREEError("Failed to create the module", Status);
+			Private::PrintIREEError("Failed to create the module", Status);
 
 			if (TempModule)
 			{
@@ -250,7 +193,7 @@ namespace UE::NNERuntimeIREE
 		Status = iree_hal_static_library_loader_create(IREE_ARRAYSIZE(LibraryList), LibraryList, iree_hal_executable_import_provider_null(), HostAllocator, &LibraryLoader);
 		if (!iree_status_is_ok(Status))
 		{
-			PrintIREEError("Failed to create the library loader", Status);
+			Private::PrintIREEError("Failed to create the library loader", Status);
 
 			if (LibraryLoader)
 			{
@@ -266,7 +209,7 @@ namespace UE::NNERuntimeIREE
 		Status = iree_hal_allocator_create_heap(Identifier, HostAllocator, HostAllocator, &DeviceAllocator);
 		if (!iree_status_is_ok(Status))
 		{
-			PrintIREEError("Failed to create the device allocator", Status);
+			Private::PrintIREEError("Failed to create the device allocator", Status);
 			
 			if (DeviceAllocator)
 			{
@@ -284,7 +227,7 @@ namespace UE::NNERuntimeIREE
 		Status = iree_hal_sync_device_create(Identifier, &DeviceParams, 1, &LibraryLoader, DeviceAllocator, HostAllocator, &TempDevice);
 		if (!iree_status_is_ok(Status))
 		{
-			PrintIREEError("Failed to create the device", Status);
+			Private::PrintIREEError("Failed to create the device", Status);
 
 			if (TempDevice)
 			{
@@ -321,7 +264,7 @@ namespace UE::NNERuntimeIREE
 		Status = iree_runtime_session_create_with_device(Instance, &SessionOptions, Device, GetHostAllocator(), &TempSession);
 		if (!iree_status_is_ok(Status))
 		{
-			PrintIREEError("Failed to create the session", Status);
+			Private::PrintIREEError("Failed to create the session", Status);
 
 			if (TempSession)
 			{
@@ -417,7 +360,7 @@ namespace UE::NNERuntimeIREE
 		Status = iree_runtime_instance_create(&InstanceOptions, iree_allocator_system(), &TempInstance);
 		if (!iree_status_is_ok(Status))
 		{
-			PrintIREEError("Failed to create the instance", Status);
+			Private::PrintIREEError("Failed to create the instance", Status);
 
 			if (TempInstance)
 			{
@@ -450,9 +393,10 @@ namespace UE::NNERuntimeIREE
 		Instance.Reset();
 	}
 
-	TSharedPtr<FIREEModule> FIREEModule::MakeModule(const FString& DirPath, const FString& VmfbFileName, const UE::NNERuntimeIREE::FModuleMetaData& ModuleMetaData)
+	TSharedPtr<FIREEModule> FIREEModule::MakeModule(const FString& DirPath, const FString& VmfbFileName, const UNNERuntimeIREEModuleMetaData& ModuleMetaData)
 	{
 		check(!VmfbFileName.IsEmpty());
+		check(!ModuleMetaData.FunctionMetaData.IsEmpty());
 
 		TSharedPtr<FIREEInstance> Instance = FIREEInstance::GetInstance();
 		if (!Instance.IsValid())
@@ -475,99 +419,34 @@ namespace UE::NNERuntimeIREE
 		iree_status_t Status = iree_ok_status();
 		check(iree_status_is_ok(Status));
 
-		iree_vm_function_t LocalMainFunction = {0};
-		TArray<UE::NNE::FTensorDesc> InputTensorDescs;
-		TArray<UE::NNE::FTensorDesc> OutputTensorDescs;
-		int32 NumMainFunctionCandidates = 0;
-		int32 Ordinal = 0;
-		while (true)
-		{
-			iree_vm_function_t MainFunctionCandidate;
-			Status = iree_vm_module_lookup_function_by_ordinal(TempModule, IREE_VM_FUNCTION_LINKAGE_EXPORT, Ordinal, &MainFunctionCandidate);
-			if (!iree_status_is_ok(Status) || iree_vm_function_is_null(MainFunctionCandidate))
-			{
-				Ordinal = -1;
-				break;
-			}
-			Ordinal++;
-
-			iree_host_size_t NumInputs = 0;
-			iree_host_size_t NumOutputs = 0;
-			iree_vm_function_signature_t Signature = iree_vm_function_signature(&MainFunctionCandidate);
-			Status = iree_vm_function_call_count_arguments_and_results(&Signature, &NumInputs, &NumOutputs);
-			if (iree_status_is_ok(Status) && NumInputs > 0 && NumOutputs > 0)
-			{
-				if (NumMainFunctionCandidates == 0)
-				{
-					iree_string_view_t FunctionName = iree_vm_function_name(&MainFunctionCandidate);
-					FString FunctionNameString = "";
-					FunctionNameString.Append(FunctionName.data, FunctionName.size);
-					if (ModuleMetaData.FunctionMetaData.Contains(FunctionNameString))
-					{
-						FFunctionMetaData MetaData = ModuleMetaData.FunctionMetaData[FunctionNameString];
-						if (MetaData.ArgumentMetaData.Num() == NumInputs && MetaData.ResultMetaData.Num() == NumOutputs)
-						{
-							LocalMainFunction = MainFunctionCandidate;
-
-							for (int32 i = 0; i < NumInputs; i++)
-							{
-								ENNETensorDataType DataType = UE::NNERuntimeIREE::Private::GetTypeFromString(MetaData.ArgumentMetaData[i].Type);
-								if (DataType == ENNETensorDataType::None)
-								{
-									UE_LOG(LogTemp, Error, TEXT("Unknown argument type in function %s: %s"), *FunctionNameString, *MetaData.ArgumentMetaData[i].Type);
-									iree_status_free(Status);
-									return TSharedPtr<FIREEModule>();
-								}
-								InputTensorDescs.Add(UE::NNE::FTensorDesc::Make(MetaData.ArgumentMetaData[i].Name, UE::NNE::FSymbolicTensorShape::Make(MetaData.ArgumentMetaData[i].Shape), DataType));
-							}
-
-							for (int32 i = 0; i < NumOutputs; i++)
-							{
-								ENNETensorDataType DataType = UE::NNERuntimeIREE::Private::GetTypeFromString(MetaData.ResultMetaData[i].Type);
-								if (DataType == ENNETensorDataType::None)
-								{
-									UE_LOG(LogTemp, Error, TEXT("Unknown result type in function %s: %s"), *FunctionNameString, *MetaData.ResultMetaData[i].Type);
-									iree_status_free(Status);
-									return TSharedPtr<FIREEModule>();
-								}
-								OutputTensorDescs.Add(UE::NNE::FTensorDesc::Make(MetaData.ResultMetaData[i].Name, UE::NNE::FSymbolicTensorShape::Make(MetaData.ResultMetaData[i].Shape), DataType));
-							}
-						}
-						else
-						{
-							UE_LOG(LogTemp, Error, TEXT("Input and output count mismatch in function %s"), *FunctionNameString);
-							iree_status_free(Status);
-							return TSharedPtr<FIREEModule>();
-						}
-					}
-					else
-					{
-						UE_LOG(LogTemp, Error, TEXT("Failed to find meta data for function %s"), *FunctionNameString);
-						iree_status_free(Status);
-						return TSharedPtr<FIREEModule>();
-					}
-				}
-				NumMainFunctionCandidates++;
-			}
-		}
-		if (NumMainFunctionCandidates < 1)
+		iree_vm_function_t MainFunction;
+		iree_string_view_t Identifier = iree_make_cstring_view(StringCast<ANSICHAR>(static_cast<const TCHAR*>(*ModuleMetaData.FunctionMetaData[0].Name)).Get());
+		Status = iree_vm_module_lookup_function_by_name(TempModule, IREE_VM_FUNCTION_LINKAGE_EXPORT, Identifier, &MainFunction);
+		if (!iree_status_is_ok(Status) || iree_vm_function_is_null(MainFunction))
 		{
-			UE_LOG(LogTemp, Error, TEXT("Failed to find a suitable module main function"));
+			UE_LOG(LogTemp, Error, TEXT("FIREEModule failed to find the module function %s"), *ModuleMetaData.FunctionMetaData[0].Name);
 			iree_status_free(Status);
 			return TSharedPtr<FIREEModule>();
 		}
-		if (NumMainFunctionCandidates > 1)
+
+		iree_host_size_t NumInputs = 0;
+		iree_host_size_t NumOutputs = 0;
+		iree_vm_function_signature_t Signature = iree_vm_function_signature(&MainFunction);
+		Status = iree_vm_function_call_count_arguments_and_results(&Signature, &NumInputs, &NumOutputs);
+		if (!iree_status_is_ok(Status) || NumInputs != ModuleMetaData.FunctionMetaData[0].InputDescs.Num() || NumOutputs != ModuleMetaData.FunctionMetaData[0].OutputDescs.Num())
 		{
-			UE_LOG(LogTemp, Warning, TEXT("Found multiple suitable module main functions"));
+			UE_LOG(LogTemp, Error, TEXT("FIREEModule has a function signature mismatch in function %s"), *ModuleMetaData.FunctionMetaData[0].Name);
+			iree_status_free(Status);
+			return TSharedPtr<FIREEModule>();
 		}
 
 		TSharedPtr<FIREEModule> Result = MakeShared<FIREEModule>();
 		Result->Instance = Instance;
 		Result->ModelData = ModelData;
 		Result->Module = TempModule;
-		Result->MainFunction = LocalMainFunction;
-		Result->InputTensorDescs = InputTensorDescs;
-		Result->OutputTensorDescs = OutputTensorDescs;
+		Result->MainFunction = MainFunction;
+		Result->InputTensorDescs = ModuleMetaData.FunctionMetaData[0].InputDescs;
+		Result->OutputTensorDescs = ModuleMetaData.FunctionMetaData[0].OutputDescs;
 
 		iree_status_free(Status);
 		return Result;
@@ -584,7 +463,7 @@ namespace UE::NNERuntimeIREE
 		Status = iree_runtime_session_append_module(Session, Module);
 		if (!iree_status_is_ok(Status))
 		{
-			PrintIREEError("Failed to append the module to the session", Status);
+			Private::PrintIREEError("Failed to append the module to the session", Status);
 			iree_status_free(Status);
 			return false;
 		}
@@ -660,7 +539,7 @@ namespace UE::NNERuntimeIREE
 		Status = iree_runtime_call_initialize(Session, MainFunction, &Call);
 		if (!iree_status_is_ok(Status))
 		{
-			PrintIREEError("Failed to initialize the session call", Status);
+			Private::PrintIREEError("Failed to initialize the session call", Status);
 			iree_status_free(Status);
 			return false;
 		}
@@ -798,7 +677,7 @@ namespace UE::NNERuntimeIREE
 				&TempBufferView);
 			if (!iree_status_is_ok(Status))
 			{
-				PrintIREEError("Failed to allocate the buffer view", Status);
+				Private::PrintIREEError("Failed to allocate the buffer view", Status);
 				if (TempBufferView)
 				{
 					iree_hal_buffer_view_release(TempBufferView);
@@ -811,7 +690,7 @@ namespace UE::NNERuntimeIREE
 			iree_hal_buffer_view_release(TempBufferView);
 			if (!iree_status_is_ok(Status))
 			{
-				PrintIREEError("Failed to push the buffer view to the input list", Status);
+				Private::PrintIREEError("Failed to push the buffer view to the input list", Status);
 				iree_status_free(Status);
 				return -1;
 			}
@@ -820,7 +699,7 @@ namespace UE::NNERuntimeIREE
 		Status = iree_runtime_call_invoke(&Call, 0);
 		if (!iree_status_is_ok(Status))
 		{
-			PrintIREEError("Failed to call the model function", Status);
+			Private::PrintIREEError("Failed to call the model function", Status);
 			iree_status_free(Status);
 			return -1;
 		}
@@ -875,7 +754,7 @@ namespace UE::NNERuntimeIREE
 				Status = iree_hal_buffer_map_range(Buffer, IREE_HAL_MAPPING_MODE_PERSISTENT, IREE_HAL_MEMORY_ACCESS_READ, 0, DataSizeInBytes, &BufferMapping);
 				if (!iree_status_is_ok(Status))
 				{
-					PrintIREEError("Failed to map the result buffer", Status);
+					Private::PrintIREEError("Failed to map the result buffer", Status);
 					Result = -1;
 					break;
 				}
diff --git a/Engine/Plugins/Experimental/NNERuntimeIREE/Source/NNERuntimeIREE/Private/NNERuntimeIREECommon.h b/Engine/Plugins/Experimental/NNERuntimeIREE/Source/NNERuntimeIREE/Private/NNERuntimeIREECommon.h
--- a/Engine/Plugins/Experimental/NNERuntimeIREE/Source/NNERuntimeIREE/Private/NNERuntimeIREECommon.h
+++ b/Engine/Plugins/Experimental/NNERuntimeIREE/Source/NNERuntimeIREE/Private/NNERuntimeIREECommon.h
@@ -5,17 +5,26 @@
 #ifdef WITH_NNE_RUNTIME_IREE
 
 #include "CoreMinimal.h"
-#include "GenericPlatform/GenericPlatformProcess.h"
-#include "HAL/FileManager.h"
-#include "Misc/Paths.h"
-#include "Misc/ScopeLock.h"
 #include "NNEModelData.h"
 #include "NNERuntimeIREECpu.h"
 #include "NNERuntimeIREEMetaData.h"
 #include "NNETypes.h"
-#include "Serialization/Archive.h"
 
-#include "iree/runtime/api.h"
+#if PLATFORM_MICROSOFT
+#include "Microsoft/AllowMicrosoftPlatformTypes.h"
+#include "Microsoft/AllowMicrosoftPlatformAtomics.h"
+#endif // PLATFORM_MICROSOFT
+THIRD_PARTY_INCLUDES_START
+#include "iree/hal/buffer.h"
+#include "iree/hal/channel.h"
+#include "iree/runtime/call.h"
+#include "iree/runtime/session.h"
+#include "iree/vm/module.h"
+THIRD_PARTY_INCLUDES_END
+#if PLATFORM_MICROSOFT
+#include "Microsoft/HideMicrosoftPlatformAtomics.h"
+#include "Microsoft/HideMicrosoftPlatformTypes.h"
+#endif // PLATFORM_MICROSOF
 
 namespace UE::NNERuntimeIREE
 {
@@ -28,7 +37,7 @@ namespace UE::NNERuntimeIREE
 		FIREEModule();
 		~FIREEModule();
 
-		static TSharedPtr<FIREEModule> MakeModule(const FString& DirPath, const FString& VmfbFileName, const UE::NNERuntimeIREE::FModuleMetaData& ModuleMetaData);
+		static TSharedPtr<FIREEModule> MakeModule(const FString& DirPath, const FString& VmfbFileName, const UNNERuntimeIREEModuleMetaData& ModuleMetaData);
 		bool AppendToSession(iree_runtime_session_t* Session);
 
 		iree_vm_function_t GetMainFunction();
diff --git a/Engine/Plugins/Experimental/NNERuntimeIREE/Source/NNERuntimeIREE/Private/NNERuntimeIREECpu.cpp b/Engine/Plugins/Experimental/NNERuntimeIREE/Source/NNERuntimeIREE/Private/NNERuntimeIREECpu.cpp
--- a/Engine/Plugins/Experimental/NNERuntimeIREE/Source/NNERuntimeIREE/Private/NNERuntimeIREECpu.cpp
+++ b/Engine/Plugins/Experimental/NNERuntimeIREE/Source/NNERuntimeIREE/Private/NNERuntimeIREECpu.cpp
@@ -22,7 +22,9 @@
 #include "Misc/Paths.h"
 #include "Misc/SecureHash.h"
 #include "NNERuntimeIREECpuModel.h"
+#include "NNERuntimeIREEMetaData.h"
 #include "Serialization/MemoryReader.h"
+#include "Serialization/MemoryWriter.h"
 
 namespace UE::NNERuntimeIREECpu::Private
 {
@@ -76,14 +78,6 @@ namespace UE::NNERuntimeIREECpu::Private
 		return RuntimeName + "-" + UNNERuntimeIREECpu::GUID.ToString(EGuidFormats::Digits) + "-" + FString::FromInt(UNNERuntimeIREECpu::Version) + "-" + FileIdString + "-" + PlatformName + (!Architecture.IsEmpty() ? ("-" + Architecture) : "");
 	}
 
-	bool GetModuleMetaData(const FString& FileType, TConstArrayView<uint8> FileData, UE::NNERuntimeIREE::FModuleMetaData* ModuleMetaData)
-	{
-		check(ModuleMetaData);
-		FString FileDataString = "";
-		FileDataString.AppendChars((char*)FileData.GetData(), FileData.Num());
-		return ModuleMetaData->ParseFromString(FileDataString);
-	}
-
 	FString GetIntermediateModelDirPath(const FString& PlatformName, const FString& ModelName)
 	{
 		return FPaths::Combine("Intermediate", "Build", PlatformName, UE_PLUGIN_NAME, ModelName);
@@ -102,7 +96,7 @@ namespace UE::NNERuntimeIREECpu::Private
 } // UE::NNERuntimeIREECpu::Private
 
 FGuid UNNERuntimeIREECpu::GUID = FGuid((int32)'I', (int32)'C', (int32)'P', (int32)'U');
-int32 UNNERuntimeIREECpu::Version = 0x00000002;
+int32 UNNERuntimeIREECpu::Version = 0x00000003;
 
 FString UNNERuntimeIREECpu::GetRuntimeName() const
 {
@@ -152,8 +146,9 @@ TSharedPtr<UE::NNE::FSharedModelData> UNNERuntimeIREECpu::CreateModelData(const
 	PlatformFile.CreateDirectoryTree(*IntermediateDir);
 
 	TArray<FIREECompilerResult> CompilerResults;
+	UNNERuntimeIREEModuleMetaData* CompilerModuleMetaData = NewObject<UNNERuntimeIREEModuleMetaData>();
 	FString StagingDir = FPaths::ConvertRelativePathToFull(FPaths::Combine(FPaths::ProjectDir(), GetStagedModelDirPath(TargetPlatformName)));
-	if (!Compiler->CompileMlir(FileData, FileIdString, IntermediateDir, StagingDir, CompilerResults))
+	if (!Compiler->CompileMlir(FileData, FileIdString, IntermediateDir, StagingDir, CompilerResults, CompilerModuleMetaData))
 	{
 		UE_LOG(LogTemp, Warning, TEXT("UNNERuntimeIREECpu failed to compile model %s"), *FileIdString);
 		return TSharedPtr<UE::NNE::FSharedModelData>();
@@ -181,14 +176,6 @@ TSharedPtr<UE::NNE::FSharedModelData> UNNERuntimeIREECpu::CreateModelData(const
 		FSharedBuffer VmfbBuffer = MakeSharedBufferFromArray(MoveTemp(VmfbData));
 		PutIntoDDC(GetModelCpuDataIdentifier(GetRuntimeName(), FileIdString, TargetPlatformName, CompilerResults[i].Architecture) + "-vmfb", VmfbBuffer);
 	}
-	
-	UE::NNERuntimeIREE::FModuleMetaData ModuleMetaData;
-	if (!GetModuleMetaData(FileType, FileData, &ModuleMetaData))
-	{
-		UE_LOG(LogTemp, Warning, TEXT("UNNERuntimeIREECpu could not get the module meta data"));
-		return TSharedPtr<UE::NNE::FSharedModelData>();
-	}
-	FString ModuleDataString = ModuleMetaData.ToJson(false);
 
 	TArray<uint8> ResultData;
 	FMemoryWriter Writer(ResultData);
@@ -197,7 +184,17 @@ TSharedPtr<UE::NNE::FSharedModelData> UNNERuntimeIREECpu::CreateModelData(const
 	FGuid FileIdCopy = FileId;
 	Writer << FileIdCopy;
 
-	Writer << ModuleDataString;
+	TArray<uint8> ModuleMetaData;
+	if (AdditionalFileData.Contains("IREEModuleMetaData"))
+	{
+		ModuleMetaData = AdditionalFileData["IREEModuleMetaData"];
+	}
+	if (ModuleMetaData.IsEmpty())
+	{
+		FMemoryWriter ObjectWriter(ModuleMetaData);
+		CompilerModuleMetaData->Serialize(ObjectWriter);
+	}
+	Writer << ModuleMetaData;
 
 	int32 NumArchitectures = CompilerResults.Num();
 	Writer << NumArchitectures;
@@ -276,8 +273,21 @@ TSharedPtr<UE::NNE::IModelCPU> UNNERuntimeIREECpu::CreateModelCPU(TObjectPtr<UNN
 	FGuid FileId = FGuid();
 	Reader << FileId;
 
-	FString ModuleDataString = "";
-	Reader << ModuleDataString;
+	TArray<uint8> ModuleDataArray;
+	Reader << ModuleDataArray;
+	if (ModuleDataArray.IsEmpty())
+	{
+		UE_LOG(LogTemp, Warning, TEXT("UNNERuntimeIREECpu failed to find any module meta data, please reimport the original model"));
+		return TSharedPtr<UE::NNE::IModelCPU>();
+	}
+	UNNERuntimeIREEModuleMetaData* ModuleMetaData = NewObject<UNNERuntimeIREEModuleMetaData>();
+	FMemoryReaderView ObjectReader(ModuleDataArray);
+	ModuleMetaData->Serialize(ObjectReader);
+	if (ModuleMetaData->FunctionMetaData.IsEmpty())
+	{
+		UE_LOG(LogTemp, Warning, TEXT("UNNERuntimeIREECpu failed to parse the module meta data, please reimport the original model"));
+		return TSharedPtr<UE::NNE::IModelCPU>();
+	}
 
 	int32 NumArchitectures = 0;
 	Reader << NumArchitectures;
@@ -326,13 +336,6 @@ TSharedPtr<UE::NNE::IModelCPU> UNNERuntimeIREECpu::CreateModelCPU(TObjectPtr<UNN
 		return TSharedPtr<UE::NNE::IModelCPU>();
 	}
 
-	UE::NNERuntimeIREE::FModuleMetaData ModuleMetaData;
-	if (!ModuleDataString.IsEmpty() && !ModuleMetaData.FromJson(ModuleDataString))
-	{
-		UE_LOG(LogTemp, Warning, TEXT("UNNERuntimeIREECpu failed to parse the module meta data"));
-		return TSharedPtr<UE::NNE::IModelCPU>();
-	}
-
 	FString FileIdString = FileId.ToString(EGuidFormats::Digits).ToLower();
 #if WITH_EDITOR
 	FString SharedLibraryDirPath = FPaths::ConvertRelativePathToFull(FPaths::Combine(FPaths::ProjectDir(), GetIntermediateModelDirPath(UGameplayStatics::GetPlatformName(), FileIdString), RelativeDirPath));
@@ -370,7 +373,7 @@ TSharedPtr<UE::NNE::IModelCPU> UNNERuntimeIREECpu::CreateModelCPU(TObjectPtr<UNN
 #endif // WITH_EDITOR
 
 	TUniquePtr<UE::NNERuntimeIREECpu::FModel> Model = MakeUnique<UE::NNERuntimeIREECpu::FModel>();
-	if (!Model->Init(SharedLibraryDirPath, SharedLibraryFileName, VmfbFileName, SharedLibraryEntryPointName, ModuleMetaData))
+	if (!Model->Init(SharedLibraryDirPath, SharedLibraryFileName, VmfbFileName, SharedLibraryEntryPointName, *ModuleMetaData))
 	{
 		UE_LOG(LogTemp, Warning, TEXT("UNNERuntimeIREECpu could not initialize the model created from model data with id %s"), *FileIdString);
 		return TSharedPtr<UE::NNE::IModelCPU>();
diff --git a/Engine/Plugins/Experimental/NNERuntimeIREE/Source/NNERuntimeIREE/Private/NNERuntimeIREECpuCompiler.cpp b/Engine/Plugins/Experimental/NNERuntimeIREE/Source/NNERuntimeIREE/Private/NNERuntimeIREECpuCompiler.cpp
--- a/Engine/Plugins/Experimental/NNERuntimeIREE/Source/NNERuntimeIREE/Private/NNERuntimeIREECpuCompiler.cpp
+++ b/Engine/Plugins/Experimental/NNERuntimeIREE/Source/NNERuntimeIREE/Private/NNERuntimeIREECpuCompiler.cpp
@@ -194,13 +194,20 @@ namespace UE::NNERuntimeIREECpu::Private
 		return MakeUnique<FNNERuntimeIREECpuCompiler>(CompilerCommand, LinkerCommand, SharedLibExt, Targets);
 	}
 
-	bool FNNERuntimeIREECpuCompiler::CompileMlir(TConstArrayView<uint8> InFileData, const FString& InModelName, const FString& InIntermediateDir, const FString& InStagingDir, TArray<FIREECompilerResult>& OutCompilerResults)
+	bool FNNERuntimeIREECpuCompiler::CompileMlir(TConstArrayView<uint8> InFileData, const FString& InModelName, const FString& InIntermediateDir, const FString& InStagingDir, TArray<FIREECompilerResult>& OutCompilerResults, UNNERuntimeIREEModuleMetaData* ModuleMetaData)
 	{
 		using namespace Private;
 
 		FString InputFilePath = FPaths::Combine(InIntermediateDir, InModelName) + ".mlir";
 		FFileHelper::SaveArrayToFile(InFileData, *InputFilePath);
 
+		if (ModuleMetaData)
+		{
+			FString FileDataString = "";
+			FileDataString.AppendChars((char*)InFileData.GetData(), InFileData.Num());
+			ModuleMetaData->ParseFromString(FileDataString);
+		}
+
 		TArray<FIREECompilerResult> Results;
 		bool bResult = true;
 		IPlatformFile& PlatformFile = FPlatformFileManager::Get().GetPlatformFile();
diff --git a/Engine/Plugins/Experimental/NNERuntimeIREE/Source/NNERuntimeIREE/Private/NNERuntimeIREECpuCompiler.h b/Engine/Plugins/Experimental/NNERuntimeIREE/Source/NNERuntimeIREE/Private/NNERuntimeIREECpuCompiler.h
--- a/Engine/Plugins/Experimental/NNERuntimeIREE/Source/NNERuntimeIREE/Private/NNERuntimeIREECpuCompiler.h
+++ b/Engine/Plugins/Experimental/NNERuntimeIREE/Source/NNERuntimeIREE/Private/NNERuntimeIREECpuCompiler.h
@@ -9,6 +9,7 @@
 #include "Containers/ArrayView.h"
 #include "Containers/UnrealString.h"
 #include "GenericPlatform/GenericPlatformMisc.h"
+#include "NNERuntimeIREEMetaData.h"
 #include "Serialization/JsonSerializerMacros.h"
 #include "Templates/UniquePtr.h"
 
@@ -59,7 +60,7 @@ namespace UE::NNERuntimeIREECpu::Private
 		static TUniquePtr<FNNERuntimeIREECpuCompiler> Make(const FString& InTargetPlatformName);
 
 	public:
-		bool CompileMlir(TConstArrayView<uint8> InFileData, const FString& InModelName, const FString& InIntermediateDir, const FString& InStagingDir, TArray<FIREECompilerResult>& OutCompilerResults);
+		bool CompileMlir(TConstArrayView<uint8> InFileData, const FString& InModelName, const FString& InIntermediateDir, const FString& InStagingDir, TArray<FIREECompilerResult>& OutCompilerResults, UNNERuntimeIREEModuleMetaData* ModuleMetaData);
 
 		// Needs to stay in header to prevent system macros to overwrite GetEnvironmentVariable
 		static FString GetEnvVar(const FString& EnvironmentVariableName) { return FPlatformMisc::GetEnvironmentVariable(*EnvironmentVariableName); }
diff --git a/Engine/Plugins/Experimental/NNERuntimeIREE/Source/NNERuntimeIREE/Private/NNERuntimeIREECpuModel.cpp b/Engine/Plugins/Experimental/NNERuntimeIREE/Source/NNERuntimeIREE/Private/NNERuntimeIREECpuModel.cpp
--- a/Engine/Plugins/Experimental/NNERuntimeIREE/Source/NNERuntimeIREE/Private/NNERuntimeIREECpuModel.cpp
+++ b/Engine/Plugins/Experimental/NNERuntimeIREE/Source/NNERuntimeIREE/Private/NNERuntimeIREECpuModel.cpp
@@ -57,12 +57,13 @@ namespace UE::NNERuntimeIREECpu
 		return Session->RunSyncCPU(InInputBindings, InOutputBindings);
 	}
 
-	bool FModel::Init(const FString& DirPath, const FString& SharedLibraryFileName, const FString& VmfbFileName, const FString& LibraryQueryFunctionName, const UE::NNERuntimeIREE::FModuleMetaData& ModuleMetaData)
+	bool FModel::Init(const FString& DirPath, const FString& SharedLibraryFileName, const FString& VmfbFileName, const FString& LibraryQueryFunctionName, const UNNERuntimeIREEModuleMetaData& ModuleMetaData)
 	{
 		check(!Module.IsValid());
 		check(!SharedLibraryFileName.IsEmpty());
 		check(!VmfbFileName.IsEmpty());
 		check(!LibraryQueryFunctionName.IsEmpty());
+		check(!ModuleMetaData.FunctionMetaData.IsEmpty());
 
 		Module = UE::NNERuntimeIREE::FIREEModule::MakeModule(DirPath, VmfbFileName, ModuleMetaData);
 		if (!Module.IsValid())
diff --git a/Engine/Plugins/Experimental/NNERuntimeIREE/Source/NNERuntimeIREE/Private/NNERuntimeIREECpuModel.h b/Engine/Plugins/Experimental/NNERuntimeIREE/Source/NNERuntimeIREE/Private/NNERuntimeIREECpuModel.h
--- a/Engine/Plugins/Experimental/NNERuntimeIREE/Source/NNERuntimeIREE/Private/NNERuntimeIREECpuModel.h
+++ b/Engine/Plugins/Experimental/NNERuntimeIREE/Source/NNERuntimeIREE/Private/NNERuntimeIREECpuModel.h
@@ -36,7 +36,7 @@ namespace UE::NNERuntimeIREECpu
 		FModel() {};
 		virtual ~FModel() = default;
 
-		bool Init(const FString& DirPath, const FString& SharedLibraryFileName, const FString& VmfbFileName, const FString& LibraryQueryFunctionName, const UE::NNERuntimeIREE::FModuleMetaData& ModuleMetaData);
+		bool Init(const FString& DirPath, const FString& SharedLibraryFileName, const FString& VmfbFileName, const FString& LibraryQueryFunctionName, const UNNERuntimeIREEModuleMetaData& ModuleMetaData);
 
 		//~ Begin IModelCPU Interface
 		virtual TSharedPtr<UE::NNE::IModelInstanceCPU> CreateModelInstanceCPU() override;
diff --git a/Engine/Plugins/Experimental/NNERuntimeIREE/Source/NNERuntimeIREE/Private/NNERuntimeIREEMetaData.cpp b/Engine/Plugins/Experimental/NNERuntimeIREE/Source/NNERuntimeIREE/Private/NNERuntimeIREEMetaData.cpp
--- a/Engine/Plugins/Experimental/NNERuntimeIREE/Source/NNERuntimeIREE/Private/NNERuntimeIREEMetaData.cpp
+++ b/Engine/Plugins/Experimental/NNERuntimeIREE/Source/NNERuntimeIREE/Private/NNERuntimeIREEMetaData.cpp
@@ -2,150 +2,335 @@
 
 #include "NNERuntimeIREEMetaData.h"
 
-namespace UE::NNERuntimeIREE
-{
-
-#ifdef WITH_NNE_RUNTIME_IREE
+#include "Serialization/CustomVersion.h"
 
-bool FArgumentMetaData::ParseFromString(const FString& ArgumentString)
+namespace UE::NNERuntimeIREE::ModuleMetaData::Private
 {
-	FString Argument = ArgumentString;
-	int32 AttributeStartIndex = Argument.Find("{");
-	if (AttributeStartIndex > 0)
-	{
-		Argument = Argument.Mid(0, AttributeStartIndex).TrimStartAndEnd();
-	}
-
-	Name = "";
-	int32 NameEnd = Argument.Find(":");
-	if (NameEnd > 0)
+	enum Version : uint32
 	{
-		Name = Argument.Mid(0, NameEnd).TrimStartAndEnd();
-		Argument = Argument.Mid(NameEnd + 1).TrimStartAndEnd();
-	}
-	else
-	{
-		Argument = Argument.TrimStartAndEnd();
-	}
+		V0 = 0, // Initial
+		// New versions can be added above this line
+		VersionPlusOne,
+		Latest = VersionPlusOne - 1
+	};
+	const FGuid GUID(0x2f9ffd31, 0x12b817cd, 0x627855bf, 0x5e405720);
+	FCustomVersionRegistration Version(GUID, Version::Latest, TEXT("NNERuntimeIREEModuleMetaDataVersion"));// Always save with the latest version
 
-	if (Argument.StartsWith("tensor"))
+	ENNETensorDataType ConvertTypeString(const FString& TypeString)
 	{
-		int32 ShapeOpenBracket = Argument.Find("<");
-		int32 ShapeCloseBracket = Argument.Find(">");
-		if (ShapeOpenBracket < 1 || ShapeCloseBracket <= ShapeOpenBracket)
+		if (TypeString.StartsWith("char"))
 		{
-			return false;
+			return ENNETensorDataType::Char;
 		}
-		FString ShapeTypeString = Argument.Mid(ShapeOpenBracket + 1, ShapeCloseBracket - (ShapeOpenBracket + 1)).TrimStartAndEnd();
-
-		TArray<FString> Dims;
-		FString Dim;
-		while (ShapeTypeString.Split("x", &Dim, &ShapeTypeString))
+		if (TypeString.StartsWith("bool") || TypeString.StartsWith("i1"))
 		{
-			Dims.Add(Dim);
+			return ENNETensorDataType::Boolean;
 		}
-		Type = ShapeTypeString;
-
-		for (int32 i = 0; i < Dims.Num(); i++)
+		else if (TypeString.StartsWith("half"))
+		{
+			return ENNETensorDataType::Half;
+		}
+		else if (TypeString.StartsWith("f"))
 		{
-			if (Dims[i].Contains("?"))
+			if (TypeString.StartsWith("f16"))
 			{
-				Shape.Add(-1);
+				return ENNETensorDataType::Half;
 			}
-			else
+			else if (TypeString.StartsWith("float") || TypeString.StartsWith("f32"))
+			{
+				return ENNETensorDataType::Float;
+			}
+			else if (TypeString.StartsWith("f64"))
 			{
-				int32 DimVal = FCString::Atoi(*Dims[i]);
-				Shape.Add(DimVal > 0 ? DimVal : -1);
+				return ENNETensorDataType::Double;
 			}
 		}
+		else if (TypeString.StartsWith("double"))
+		{
+			return ENNETensorDataType::Double;
+		}
+		else if (TypeString.StartsWith("i") || TypeString.StartsWith("si"))
+		{
+			if (TypeString.EndsWith("i8"))
+			{
+				return ENNETensorDataType::Int8;
+			}
+			else if (TypeString.EndsWith("i16"))
+			{
+				return ENNETensorDataType::Int16;
+			}
+			else if (TypeString.EndsWith("i32") || TypeString.EndsWith("int"))
+			{
+				return ENNETensorDataType::Int32;
+			}
+			else if (TypeString.EndsWith("i64"))
+			{
+				return ENNETensorDataType::Int64;
+			}
+		}
+		else if (TypeString.StartsWith("ui"))
+		{
+			if (TypeString.EndsWith("i8"))
+			{
+				return ENNETensorDataType::UInt8;
+			}
+			else if (TypeString.EndsWith("i16"))
+			{
+				return ENNETensorDataType::UInt16;
+			}
+			else if (TypeString.EndsWith("i32"))
+			{
+				return ENNETensorDataType::UInt32;
+			}
+			else if (TypeString.EndsWith("i64"))
+			{
+				return ENNETensorDataType::UInt64;
+			}
+		}
+		return ENNETensorDataType::None;
 	}
-	else
-	{
-		Shape.SetNumUninitialized(0);
-		Type = Argument;
-	}
-
-	return true;
-}
 
-bool FFunctionMetaData::ParseFromString(const FString& FunctionString)
-{
-	FString ResultPattern = "->";
-	int32 ArgumentsOpenBracket = FunctionString.Find("(");
-	int32 ArgumentsCloseBracket = FunctionString.Find(")");
-	int32 ResultStart = FunctionString.Find(ResultPattern);
-	if (ArgumentsOpenBracket < 1 || ArgumentsCloseBracket <= ArgumentsOpenBracket)
+	bool ParseTensorDescFromString(const FString& ArgumentString, FString& Name, TArray<int32>& Shape, ENNETensorDataType& Type)
 	{
-		return false;
-	}
-
-	Name = FunctionString.Mid(0, ArgumentsOpenBracket).TrimStartAndEnd();
+		FString Argument = ArgumentString;
+		int32 AttributeStartIndex = Argument.Find("{");
+		if (AttributeStartIndex > 0)
+		{
+			Argument = Argument.Mid(0, AttributeStartIndex).TrimStartAndEnd();
+		}
 
-	FString ArgumentsString = FunctionString.Mid(ArgumentsOpenBracket + 1, ArgumentsCloseBracket - (ArgumentsOpenBracket + 1) ).TrimStartAndEnd();
-	while (ArgumentsString.Len() > 0)
-	{
-		int32 SeparatorIndex = ArgumentsString.Find(",");
-		FString ArgumentString;
-		if (SeparatorIndex > 0)
+		Name = "";
+		int32 NameEnd = Argument.Find(":");
+		if (NameEnd > 0)
 		{
-			ArgumentString = ArgumentsString.Mid(0, SeparatorIndex).TrimStartAndEnd();
-			ArgumentsString = ArgumentsString.Mid(SeparatorIndex + 1).TrimStartAndEnd();
+			Name = Argument.Mid(0, NameEnd).TrimStartAndEnd();
+			Argument = Argument.Mid(NameEnd + 1).TrimStartAndEnd();
 		}
 		else
 		{
-			ArgumentString = ArgumentsString;
-			ArgumentsString = "";
+			Argument = Argument.TrimStartAndEnd();
 		}
 
-		FArgumentMetaData MetaData;
-		if (MetaData.ParseFromString(ArgumentString))
+		Shape.Empty();
+		if (Argument.StartsWith("tensor"))
 		{
-			ArgumentMetaData.Add(MetaData);
+			int32 ShapeOpenBracket = Argument.Find("<");
+			int32 ShapeCloseBracket = Argument.Find(">");
+			if (ShapeOpenBracket < 1 || ShapeCloseBracket <= ShapeOpenBracket)
+			{
+				return false;
+			}
+			FString ShapeTypeString = Argument.Mid(ShapeOpenBracket + 1, ShapeCloseBracket - (ShapeOpenBracket + 1)).TrimStartAndEnd();
+
+			TArray<FString> Dims;
+			FString Dim;
+			while (ShapeTypeString.Split("x", &Dim, &ShapeTypeString))
+			{
+				Dims.Add(Dim);
+			}
+			Type = ConvertTypeString(ShapeTypeString);
+
+			for (int32 i = 0; i < Dims.Num(); i++)
+			{
+				if (Dims[i].Contains("?"))
+				{
+					Shape.Add(-1);
+				}
+				else
+				{
+					int32 DimVal = FCString::Atoi(*Dims[i]);
+					Shape.Add(DimVal > 0 ? DimVal : -1);
+				}
+			}
 		}
 		else
 		{
-			return false;
+			Shape.Add(1);
+			Type = ConvertTypeString(Argument);
 		}
+
+		return Type != ENNETensorDataType::None;
 	}
 
-	if (ResultStart > ArgumentsCloseBracket)
+	bool ParseFunctionMetaDataFromString(const FString& FunctionString, UE::NNERuntimeIREE::FFunctionMetaData& FunctionMetaData)
 	{
-		FString ResultsString = FunctionString.Mid(ResultStart + ResultPattern.Len()).TrimStartAndEnd();
-		ResultsString.RemoveFromStart("(");
-		ResultsString.RemoveFromEnd(")");
-		ResultsString = ResultsString.TrimStartAndEnd();
-		while (ResultsString.Len() > 0)
-		{
-			int32 SeparatorIndex = ResultsString.Find(",");
-			FString ResultString;
+		FString ResultPattern = "->";
+		int32 ArgumentsOpenBracket = FunctionString.Find("(");
+		int32 ArgumentsCloseBracket = FunctionString.Find(")");
+		int32 ResultStart = FunctionString.Find(ResultPattern);
+		if (ArgumentsOpenBracket < 1 || ArgumentsCloseBracket <= ArgumentsOpenBracket)
+		{
+			return false;
+		}
+
+		FunctionMetaData.Name = FunctionString.Mid(0, ArgumentsOpenBracket).TrimStartAndEnd();
+		FunctionMetaData.InputDescs.Empty();
+		FunctionMetaData.OutputDescs.Empty();
+
+		FString ArgumentsString = FunctionString.Mid(ArgumentsOpenBracket + 1, ArgumentsCloseBracket - (ArgumentsOpenBracket + 1)).TrimStartAndEnd();
+		while (ArgumentsString.Len() > 0)
+		{
+			int32 SeparatorIndex = ArgumentsString.Find(",");
+			FString ArgumentString;
 			if (SeparatorIndex > 0)
 			{
-				ResultString = ResultsString.Mid(0, SeparatorIndex).TrimStartAndEnd();
-				ResultsString = ResultsString.Mid(SeparatorIndex + 1).TrimStartAndEnd();
+				ArgumentString = ArgumentsString.Mid(0, SeparatorIndex).TrimStartAndEnd();
+				ArgumentsString = ArgumentsString.Mid(SeparatorIndex + 1).TrimStartAndEnd();
 			}
 			else
 			{
-				ResultString = ResultsString;
-				ResultsString = "";
+				ArgumentString = ArgumentsString;
+				ArgumentsString = "";
 			}
 
-			FArgumentMetaData MetaData;
-			if (MetaData.ParseFromString(ResultString))
+			FString Name;
+			TArray<int32> Shape;
+			ENNETensorDataType Type;
+			if (ParseTensorDescFromString(ArgumentString, Name, Shape, Type))
 			{
-				ResultMetaData.Add(MetaData);
+				FunctionMetaData.InputDescs.Add(UE::NNE::FTensorDesc::Make(Name, UE::NNE::FSymbolicTensorShape::Make(Shape), Type));
 			}
 			else
 			{
 				return false;
 			}
 		}
+
+		if (ResultStart > ArgumentsCloseBracket)
+		{
+			FString ResultsString = FunctionString.Mid(ResultStart + ResultPattern.Len()).TrimStartAndEnd();
+			ResultsString.RemoveFromStart("(");
+			ResultsString.RemoveFromEnd(")");
+			ResultsString = ResultsString.TrimStartAndEnd();
+			while (ResultsString.Len() > 0)
+			{
+				int32 SeparatorIndex = ResultsString.Find(",");
+				FString ResultString;
+				if (SeparatorIndex > 0)
+				{
+					ResultString = ResultsString.Mid(0, SeparatorIndex).TrimStartAndEnd();
+					ResultsString = ResultsString.Mid(SeparatorIndex + 1).TrimStartAndEnd();
+				}
+				else
+				{
+					ResultString = ResultsString;
+					ResultsString = "";
+				}
+
+				FString Name;
+				TArray<int32> Shape;
+				ENNETensorDataType Type;
+				if (ParseTensorDescFromString(ResultString, Name, Shape, Type))
+				{
+					FunctionMetaData.OutputDescs.Add(UE::NNE::FTensorDesc::Make(Name, UE::NNE::FSymbolicTensorShape::Make(Shape), Type));
+				}
+				else
+				{
+					return false;
+				}
+			}
+		}
+
+		return true;
+	}
+}
+
+void UNNERuntimeIREEModuleMetaData::Serialize(FArchive& Ar)
+{
+	// Store the asset version (no effect in load)
+	Ar.UsingCustomVersion(UE::NNERuntimeIREE::ModuleMetaData::Private::GUID);
+
+	if (Ar.IsSaving() || Ar.IsCountingMemory()) 
+	{
+		int32 NumItems = FunctionMetaData.Num();
+		Ar << NumItems;
+		for (int32 i = 0; i < NumItems; i++)
+		{
+			Ar << FunctionMetaData[i].Name;
+
+			int32 NumInputs = FunctionMetaData[i].InputDescs.Num();
+			Ar << NumInputs;
+			for (int32 j = 0; j < NumInputs; j++)
+			{
+				FString Name = FunctionMetaData[i].InputDescs[j].GetName();
+				Ar << Name;
+
+				ENNETensorDataType Type = FunctionMetaData[i].InputDescs[j].GetDataType();
+				Ar << Type;
+
+				TArray<int32> Shape = (TArray<int32>)FunctionMetaData[i].InputDescs[j].GetShape().GetData();
+				Ar << Shape;
+			}
+
+			int32 NumOutputs = FunctionMetaData[i].OutputDescs.Num();
+			Ar << NumOutputs;
+			for (int32 j = 0; j < NumOutputs; j++)
+			{
+				FString Name = FunctionMetaData[i].OutputDescs[j].GetName();
+				Ar << Name;
+
+				ENNETensorDataType Type = FunctionMetaData[i].OutputDescs[j].GetDataType();
+				Ar << Type;
+
+				TArray<int32> Shape = (TArray<int32>)FunctionMetaData[i].OutputDescs[j].GetShape().GetData();
+				Ar << Shape;
+			}
+		}
 	}
+	else
+	{
+		int32 NumItems = 0;
+		int32 NumInputs = 0;
+		int32 NumOutputs = 0;
+		UE::NNERuntimeIREE::FFunctionMetaData MetaData;
+		FString Name;
+		ENNETensorDataType Type;
+		TArray<int32> Shape;
+
+		switch (Ar.CustomVer(UE::NNERuntimeIREE::ModuleMetaData::Private::GUID))
+		{
+		case UE::NNERuntimeIREE::ModuleMetaData::Private::Version::V0:
+			Ar << NumItems;
+			FunctionMetaData.SetNum(NumItems, true);
+			for (int32 i = 0; i < NumItems; i++)
+			{
+				Ar << MetaData.Name;
+
+				Ar << NumInputs;
+				MetaData.InputDescs.Empty();
+				for (int32 j = 0; j < NumInputs; j++)
+				{
+					Ar << Name;
+					Ar << Type;
+					Ar << Shape;
+					MetaData.InputDescs.Add(UE::NNE::FTensorDesc::Make(Name, UE::NNE::FSymbolicTensorShape::Make(Shape), Type));
+				}
 
-	return true;
+				Ar << NumOutputs;
+				MetaData.OutputDescs.Empty();
+				for (int32 j = 0; j < NumOutputs; j++)
+				{
+					Ar << Name;
+					Ar << Type;
+					Ar << Shape;
+					MetaData.OutputDescs.Add(UE::NNE::FTensorDesc::Make(Name, UE::NNE::FSymbolicTensorShape::Make(Shape), Type));
+				}
+
+				FunctionMetaData[i] = MetaData;
+			}
+			break;
+		default:
+			UE_LOG(LogTemp, Error, TEXT("UNNERuntimeIREEModuleMetaData: Unknown asset version %d: Deserialisation failed, please reimport the original model."), Ar.CustomVer(UE::NNERuntimeIREE::ModuleMetaData::Private::GUID));
+			break;
+		}
+	}
 }
 
-bool FModuleMetaData::ParseFromString(const FString& ModuleString)
+bool UNNERuntimeIREEModuleMetaData::ParseFromString(const FString& ModuleString)
 {
+	using namespace UE::NNERuntimeIREE::ModuleMetaData::Private;
+
+	TArray<UE::NNERuntimeIREE::FFunctionMetaData> Result;
 	FString SearchString = ModuleString;
 	FString Pattern = "func.func";
 	int32 MatchStart = -1;
@@ -178,33 +363,18 @@ bool FModuleMetaData::ParseFromString(const FString& ModuleString)
 				MatchStart = ArgumentsEnd;
 			}
 
-			FFunctionMetaData MetaData;
-			if (MetaData.ParseFromString(SearchString.Mid(0, MatchStart)))
+			UE::NNERuntimeIREE::FFunctionMetaData MetaData;
+			if (!ParseFunctionMetaDataFromString(SearchString.Mid(0, MatchStart), MetaData))
 			{
-				FunctionMetaData.Emplace(MetaData.Name, MetaData);
+				return false;
 			}
+			Result.Add(MetaData);
 		}
 	}
-	return FunctionMetaData.Num() > 0;
-}
-
-#else
-
-bool FArgumentMetaData::ParseFromString(const FString& ArgumentString)
-{
-	return false;
-}
-
-bool FFunctionMetaData::ParseFromString(const FString& FunctionString)
-{
-	return false;
-}
-
-bool FModuleMetaData::ParseFromString(const FString& ModuleString)
-{
+	if (!Result.IsEmpty())
+	{
+		FunctionMetaData = Result;
+		return true;
+	}
 	return false;
-}
-
-#endif // WITH_NNE_RUNTIME_IREE
-
-} // UE::NNERuntimeIREE
\ No newline at end of file
+}
\ No newline at end of file
diff --git a/Engine/Plugins/Experimental/NNERuntimeIREE/Source/NNERuntimeIREE/Private/NNERuntimeIREEMetaData.h b/Engine/Plugins/Experimental/NNERuntimeIREE/Source/NNERuntimeIREE/Private/NNERuntimeIREEMetaData.h
deleted file mode 100644
diff --git a/Engine/Plugins/Experimental/NNERuntimeIREE/Source/NNERuntimeIREEEditor/NNERuntimeIREEEditor.Build.cs b/Engine/Plugins/Experimental/NNERuntimeIREE/Source/NNERuntimeIREEEditor/NNERuntimeIREEEditor.Build.cs
--- a/Engine/Plugins/Experimental/NNERuntimeIREE/Source/NNERuntimeIREEEditor/NNERuntimeIREEEditor.Build.cs
+++ b/Engine/Plugins/Experimental/NNERuntimeIREE/Source/NNERuntimeIREEEditor/NNERuntimeIREEEditor.Build.cs
@@ -8,22 +8,17 @@ public class NNERuntimeIREEEditor : ModuleRules
 	{
 		PCHUsage = PCHUsageMode.UseExplicitOrSharedPCHs;
 
-		PublicDependencyModuleNames.AddRange
-		(
-			new string[] 
-			{
-				"Core",
-				"CoreUObject",
-				"Engine"
-			}
-		);
-
 		PrivateDependencyModuleNames.AddRange
 		(
 			new string[]
 			{
-				"UnrealEd",
-				"NNE"
+				"Core",
+				"CoreUObject",
+				"Engine",
+				"MainFrame",
+				"NNE",
+				"NNERuntimeIREE",
+				"UnrealEd"
 			}
 		);
 	}
diff --git a/Engine/Plugins/Experimental/NNERuntimeIREE/Source/NNERuntimeIREEEditor/Private/NNERuntimeIREEModelDataFactory.cpp b/Engine/Plugins/Experimental/NNERuntimeIREE/Source/NNERuntimeIREEEditor/Private/NNERuntimeIREEModelDataFactory.cpp
--- a/Engine/Plugins/Experimental/NNERuntimeIREE/Source/NNERuntimeIREEEditor/Private/NNERuntimeIREEModelDataFactory.cpp
+++ b/Engine/Plugins/Experimental/NNERuntimeIREE/Source/NNERuntimeIREEEditor/Private/NNERuntimeIREEModelDataFactory.cpp
@@ -2,17 +2,17 @@
 
 #include "NNERuntimeIREEModelDataFactory.h"
 
+#include "CoreMinimal.h"
 #include "Editor.h"
-#include "Editor/EditorEngine.h"
-#include "EditorFramework/AssetImportData.h"
 #include "EngineAnalytics.h"
-#include "HAL/FileManager.h"
+#include "Interfaces/IMainFrameModule.h"
 #include "Kismet/GameplayStatics.h"
-#include "NNE.h"
+#include "Modules/ModuleManager.h"
 #include "NNEModelData.h"
+#include "NNERuntimeIREEMetaData.h"
+#include "Serialization/MemoryWriter.h"
 #include "Subsystems/ImportSubsystem.h"
 
-
 UNNERuntimeIREEModelDataFactory::UNNERuntimeIREEModelDataFactory(const FObjectInitializer& ObjectInitializer) : UFactory(ObjectInitializer)
 {
 	bCreateNew = false;
@@ -32,9 +32,26 @@ UObject* UNNERuntimeIREEModelDataFactory::FactoryCreateBinary(UClass* Class, UOb
 		return nullptr;
 	}
 
+	TConstArrayView<uint8> BufferView = MakeArrayView(Buffer, BufferEnd - Buffer);
+	FString FileDataString = "";
+	FileDataString.AppendChars((char*)BufferView.GetData(), BufferView.Num());
+
+	UNNERuntimeIREEModuleMetaData* ModuleMetaData = NewObject<UNNERuntimeIREEModuleMetaData>();
+	if (!ModuleMetaData->ParseFromString(FileDataString) || ModuleMetaData->FunctionMetaData.IsEmpty())
+	{
+		UE_LOG(LogTemp, Error, TEXT("UNNERuntimeIREEModelDataFactory failed to parse the models meta data"));
+		return nullptr;
+	}
+
+	TArray<uint8> MetaDataByteArray;
+	FMemoryWriter Writer(MetaDataByteArray);
+	ModuleMetaData->Serialize(Writer);
+
+	TMap<FString, TConstArrayView<uint8>> AdditionalFileData;
+	AdditionalFileData.Add("IREEModuleMetaData", MetaDataByteArray);
+
 	UNNEModelData* ModelData = NewObject<UNNEModelData>(InParent, Class, Name, Flags);
-	TConstArrayView<uint8> BufferView = MakeArrayView(Buffer, BufferEnd-Buffer);
-	ModelData->Init(Type, BufferView);
+	ModelData->Init(Type, BufferView, AdditionalFileData);
 
 	GEditor->GetEditorSubsystem<UImportSubsystem>()->BroadcastAssetPostImport(this, ModelData);
 
diff --git a/Engine/Plugins/Experimental/NNERuntimeIREE/Source/NNERuntimeIREEEditor/Public/NNERuntimeIREEModelDataFactory.h b/Engine/Plugins/Experimental/NNERuntimeIREE/Source/NNERuntimeIREEEditor/Public/NNERuntimeIREEModelDataFactory.h
--- a/Engine/Plugins/Experimental/NNERuntimeIREE/Source/NNERuntimeIREEEditor/Public/NNERuntimeIREEModelDataFactory.h
+++ b/Engine/Plugins/Experimental/NNERuntimeIREE/Source/NNERuntimeIREEEditor/Public/NNERuntimeIREEModelDataFactory.h
@@ -2,8 +2,7 @@
 
 #pragma once
 
-#include "CoreMinimal.h"
-#include "EditorReimportHandler.h"
+#include "UObject/UObjectGlobals.h"
 #include "Factories/Factory.h"
 
 #include "NNERuntimeIREEModelDataFactory.generated.h"
