[NNE][NNERuntimeIREE] Adding an import dialogue when importing mlir based neural network models

Problem: In contrast to other runtimes, MLIR based models can contain multiple functions of which one needs to be selected as model evaluation function.

Solution: Add a dialogue on import with a dropdown for the user to select the main function, if the models has more than one.

Tests: Imported a mlir module with multiple function definitions, selected the non-default on import and PIE in editor

#rnx

diff --git a/Engine/Plugins/Experimental/NNERuntimeIREE/Source/NNERuntimeIREE/Private/NNERuntimeIREEMetaData.cpp b/Engine/Plugins/Experimental/NNERuntimeIREE/Source/NNERuntimeIREE/Private/NNERuntimeIREEMetaData.cpp
--- a/Engine/Plugins/Experimental/NNERuntimeIREE/Source/NNERuntimeIREE/Private/NNERuntimeIREEMetaData.cpp
+++ b/Engine/Plugins/Experimental/NNERuntimeIREE/Source/NNERuntimeIREE/Private/NNERuntimeIREEMetaData.cpp
@@ -164,6 +164,10 @@ namespace UE::NNERuntimeIREE::ModuleMetaData::Private
 		}
 
 		FunctionMetaData.Name = FunctionString.Mid(0, ArgumentsOpenBracket).TrimStartAndEnd();
+		if (FunctionMetaData.Name.IsEmpty())
+		{
+			return false;
+		}
 		FunctionMetaData.InputDescs.Empty();
 		FunctionMetaData.OutputDescs.Empty();
 
diff --git a/Engine/Plugins/Experimental/NNERuntimeIREE/Source/NNERuntimeIREEEditor/NNERuntimeIREEEditor.Build.cs b/Engine/Plugins/Experimental/NNERuntimeIREE/Source/NNERuntimeIREEEditor/NNERuntimeIREEEditor.Build.cs
--- a/Engine/Plugins/Experimental/NNERuntimeIREE/Source/NNERuntimeIREEEditor/NNERuntimeIREEEditor.Build.cs
+++ b/Engine/Plugins/Experimental/NNERuntimeIREE/Source/NNERuntimeIREEEditor/NNERuntimeIREEEditor.Build.cs
@@ -18,6 +18,8 @@ public class NNERuntimeIREEEditor : ModuleRules
 				"MainFrame",
 				"NNE",
 				"NNERuntimeIREE",
+				"Slate",
+				"SlateCore",
 				"UnrealEd"
 			}
 		);
diff --git a/Engine/Plugins/Experimental/NNERuntimeIREE/Source/NNERuntimeIREEEditor/Private/NNERuntimeIREEModelDataFactory.cpp b/Engine/Plugins/Experimental/NNERuntimeIREE/Source/NNERuntimeIREEEditor/Private/NNERuntimeIREEModelDataFactory.cpp
--- a/Engine/Plugins/Experimental/NNERuntimeIREE/Source/NNERuntimeIREEEditor/Private/NNERuntimeIREEModelDataFactory.cpp
+++ b/Engine/Plugins/Experimental/NNERuntimeIREE/Source/NNERuntimeIREEEditor/Private/NNERuntimeIREEModelDataFactory.cpp
@@ -5,6 +5,10 @@
 #include "CoreMinimal.h"
 #include "Editor.h"
 #include "EngineAnalytics.h"
+#include "Framework/Application/SlateApplication.h"
+#include "Framework/Commands/UIAction.h"
+#include "Framework/MultiBox/MultiBoxBuilder.h"
+#include "Input/Reply.h"
 #include "Interfaces/IMainFrameModule.h"
 #include "Kismet/GameplayStatics.h"
 #include "Modules/ModuleManager.h"
@@ -12,6 +16,126 @@
 #include "NNERuntimeIREEMetaData.h"
 #include "Serialization/MemoryWriter.h"
 #include "Subsystems/ImportSubsystem.h"
+#include "Textures/SlateIcon.h"
+#include "Widgets/Input/SButton.h"
+#include "Widgets/Input/SComboButton.h"
+#include "Widgets/Layout/SBox.h"
+#include "Widgets/SBoxPanel.h"
+#include "Widgets/SCompoundWidget.h"
+#include "Widgets/SWidget.h"
+#include "Widgets/SWindow.h"
+#include "Widgets/Text/STextBlock.h"
+
+class SMLIRImportWindow : public SCompoundWidget
+{
+public:
+	SLATE_BEGIN_ARGS(SMLIRImportWindow) {}
+		SLATE_ARGUMENT(TConstArrayView<UE::NNERuntimeIREE::FFunctionMetaData>, MetaData)
+		SLATE_ARGUMENT(TSharedPtr<SWindow>, WidgetWindow)
+	SLATE_END_ARGS()
+
+	void Construct(const FArguments& InArgs);
+
+	TSharedRef<SWidget> OnSelectFunctionContent();
+	void OnSelectFunctionChange(int32 Index);
+	FText GetCurrentSelectedFunction() const;
+	int32 GetCurrentSelectedFunctionIndex() const;
+
+	FReply OnImportButtonClicked();
+	bool ImportButtonClicked() const;
+
+private:
+	bool bImportButtonClicked = false;
+	int32 Selection = 0;
+	TConstArrayView<UE::NNERuntimeIREE::FFunctionMetaData> MetaData;
+	TWeakPtr<SWindow> WidgetWindow;
+};
+
+void SMLIRImportWindow::Construct(const FArguments& InArgs)
+{
+	bImportButtonClicked = false;
+	Selection = 0;
+	MetaData = InArgs._MetaData;
+	WidgetWindow = InArgs._WidgetWindow;
+
+	ChildSlot
+	[
+		SNew(SBox)
+		.Padding(FMargin(10.0f, 10.0f))
+		[
+			SNew(SVerticalBox)
+			+SVerticalBox::Slot()
+			[
+				SNew(STextBlock)
+				.Text(NSLOCTEXT("UnrealEd", "NNERuntimeIREEImportFunctionSelectionText", "Neural network main function:"))
+			]
+			+SVerticalBox::Slot()
+			[
+				SNew(SComboButton)
+				.OnGetMenuContent(this, &SMLIRImportWindow::OnSelectFunctionContent)
+				.ButtonContent()
+				[
+					SNew(STextBlock)
+					.Text(this, &SMLIRImportWindow::GetCurrentSelectedFunction)
+				]
+			]
+			+SVerticalBox::Slot()
+			.Padding(FMargin(0.0f, 10.0f))
+			[
+				SNew(SButton)
+				.HAlign(HAlign_Center)
+				.Text(NSLOCTEXT("UnrealEd", "NNERuntimeIREEImportButton", "Import"))
+				.OnClicked(this, &SMLIRImportWindow::OnImportButtonClicked)
+			]
+		]
+		
+	];
+}
+
+TSharedRef<SWidget> SMLIRImportWindow::OnSelectFunctionContent()
+{
+	FMenuBuilder MenuBuilder(true, NULL);
+	for (int32 i = 0; i < MetaData.Num(); i++)
+	{
+		FUIAction ItemAction(FExecuteAction::CreateSP(this, &SMLIRImportWindow::OnSelectFunctionChange, i));
+		MenuBuilder.AddMenuEntry(FText::FromName(FName(MetaData[i].Name)), TAttribute<FText>(), FSlateIcon(), ItemAction);
+	}
+	return MenuBuilder.MakeWidget();
+}
+
+void SMLIRImportWindow::OnSelectFunctionChange(int32 Index)
+{
+	Selection = Index;
+}
+
+FText SMLIRImportWindow::GetCurrentSelectedFunction() const
+{
+	if (Selection >= 0 && Selection < MetaData.Num())
+	{
+		return FText::FromString(MetaData[Selection].Name);
+	}
+	return FText();
+}
+
+int32 SMLIRImportWindow::GetCurrentSelectedFunctionIndex() const
+{
+	return Selection;
+}
+
+FReply SMLIRImportWindow::OnImportButtonClicked()
+{
+	bImportButtonClicked = true;
+	if (WidgetWindow.IsValid())
+	{
+		WidgetWindow.Pin()->RequestDestroyWindow();
+	}
+	return FReply::Handled();
+}
+
+bool SMLIRImportWindow::ImportButtonClicked() const
+{
+	return bImportButtonClicked;
+}
 
 UNNERuntimeIREEModelDataFactory::UNNERuntimeIREEModelDataFactory(const FObjectInitializer& ObjectInitializer) : UFactory(ObjectInitializer)
 {
@@ -33,16 +157,51 @@ UObject* UNNERuntimeIREEModelDataFactory::FactoryCreateBinary(UClass* Class, UOb
 	}
 
 	TConstArrayView<uint8> BufferView = MakeArrayView(Buffer, BufferEnd - Buffer);
-	FString FileDataString = "";
+	FString FileDataString;
 	FileDataString.AppendChars((char*)BufferView.GetData(), BufferView.Num());
 
-	UNNERuntimeIREEModuleMetaData* ModuleMetaData = NewObject<UNNERuntimeIREEModuleMetaData>();
+	TObjectPtr<UNNERuntimeIREEModuleMetaData> ModuleMetaData = NewObject<UNNERuntimeIREEModuleMetaData>();
 	if (!ModuleMetaData->ParseFromString(FileDataString) || ModuleMetaData->FunctionMetaData.IsEmpty())
 	{
 		UE_LOG(LogTemp, Error, TEXT("UNNERuntimeIREEModelDataFactory failed to parse the models meta data"));
 		return nullptr;
 	}
 
+	if (ModuleMetaData->FunctionMetaData.Num() > 1)
+	{
+		TSharedPtr<SWindow> ParentWindow;
+		if (FModuleManager::Get().IsModuleLoaded("MainFrame"))
+		{
+			IMainFrameModule& MainFrame = FModuleManager::LoadModuleChecked<IMainFrameModule>("MainFrame");
+			ParentWindow = MainFrame.GetParentWindow();
+		}
+
+		TSharedRef<SWindow> Window = SNew(SWindow)
+			.Title(NSLOCTEXT("UnrealEd", "NNERuntimeIREEImportWindowTitle", "MLIR Import Options"))
+			.SizingRule(ESizingRule::Autosized)
+			.AutoCenter(EAutoCenter::PreferredWorkArea);
+
+		TSharedPtr<SMLIRImportWindow> ImportWindow;
+		Window->SetContent
+		(
+			SAssignNew(ImportWindow, SMLIRImportWindow)
+			.MetaData(ModuleMetaData->FunctionMetaData)
+			.WidgetWindow(Window)
+		);
+
+		FSlateApplication::Get().AddModalWindow(Window, ParentWindow, false);
+
+		if (!ImportWindow->ImportButtonClicked())
+		{
+			UE_LOG(LogTemp, Error, TEXT("UNNERuntimeIREEModelDataFactory could not import the model! Please select the neural network main function in the import dialog!"));
+			return nullptr;
+		}
+
+		UE::NNERuntimeIREE::FFunctionMetaData Swapper = ModuleMetaData->FunctionMetaData[0];
+		ModuleMetaData->FunctionMetaData[0] = ModuleMetaData->FunctionMetaData[ImportWindow->GetCurrentSelectedFunctionIndex()];
+		ModuleMetaData->FunctionMetaData[ImportWindow->GetCurrentSelectedFunctionIndex()] = Swapper;
+	}
+
 	TArray<uint8> MetaDataByteArray;
 	FMemoryWriter Writer(MetaDataByteArray);
 	ModuleMetaData->Serialize(Writer);
