[NNE][NNERuntimeIREE] Cleaning up pass

Problem: Currently, IREE is exposed inside headers as part of the implementing classes and thus all IREE headers are included. Also, classes and namespace are not coded consistently making the code hard to read.

Solution: Split interface and implementation so that IREE is contained only inside the source file where it is actually being used and make classes and namespaces consistent to improve readability.

#rnx
#tests Compiled and packaged locally (WIn64) for desktop and console platforms

diff --git a/Engine/Plugins/Experimental/NNERuntimeIREE/Source/NNERuntimeIREE/Internal/NNERuntimeIREEMetaData.h b/Engine/Plugins/Experimental/NNERuntimeIREE/Source/NNERuntimeIREE/Internal/NNERuntimeIREEMetaData.h
--- a/Engine/Plugins/Experimental/NNERuntimeIREE/Source/NNERuntimeIREE/Internal/NNERuntimeIREEMetaData.h
+++ b/Engine/Plugins/Experimental/NNERuntimeIREE/Source/NNERuntimeIREE/Internal/NNERuntimeIREEMetaData.h
@@ -2,14 +2,15 @@
 
 #pragma once
 
-#include "CoreMinimal.h"
+#include "Containers/ContainersFwd.h"
 #include "NNETypes.h"
+#include "UObject/Class.h"
 
 #include "NNERuntimeIREEMetaData.generated.h"
 
 namespace UE::NNERuntimeIREE
 {
-	struct FFunctionMetaData
+	struct NNERUNTIMEIREE_API FFunctionMetaData
 	{
 		FString Name;
 		TArray<UE::NNE::FTensorDesc> InputDescs;
diff --git a/Engine/Plugins/Experimental/NNERuntimeIREE/Source/NNERuntimeIREE/Private/NNERuntimeIREE.cpp b/Engine/Plugins/Experimental/NNERuntimeIREE/Source/NNERuntimeIREE/Private/NNERuntimeIREE.cpp
new file mode 100644
--- /dev/null
+++ b/Engine/Plugins/Experimental/NNERuntimeIREE/Source/NNERuntimeIREE/Private/NNERuntimeIREE.cpp
@@ -0,0 +1,633 @@
+// Copyright Epic Games, Inc. All Rights Reserved.
+
+#include "NNERuntimeIREE.h"
+
+#ifdef WITH_NNE_RUNTIME_IREE
+
+#if WITH_EDITOR
+#include "Containers/StringConv.h"
+#include "DerivedDataCache.h"
+#include "DerivedDataRequestOwner.h"
+#include "HAL/PlatformFileManager.h"
+#include "IO/IoHash.h"
+#include "Memory/SharedBuffer.h"
+#include "Misc/FileHelper.h"
+#endif // WITH_EDITOR
+
+#include "EngineAnalytics.h"
+#include "HAL/Platform.h"
+#include "Interfaces/ITargetPlatform.h"
+#include "Kismet/GameplayStatics.h"
+#include "Misc/App.h"
+#include "Misc/Paths.h"
+#include "Misc/SecureHash.h"
+#include "NNERuntimeIREECompiler.h"
+#include "NNERuntimeIREEModel.h"
+#include "NNERuntimeIREEMetaData.h"
+#include "NNEModelData.h"
+#include "Serialization/MemoryReader.h"
+#include "Serialization/MemoryWriter.h"
+
+namespace UE::NNERuntimeIREE::CPU::Private
+{
+#if WITH_EDITOR
+	inline UE::DerivedData::FCacheKey CreateCacheKey(const FString& RequestId)
+	{
+		return { UE::DerivedData::FCacheBucket(TEXT("NNEModelData")), FIoHash::HashBuffer(MakeMemoryView(StringCast<UTF8CHAR>(*RequestId))) };
+	}
+
+	inline void PutIntoDDC(const FString& RequestId, const FSharedBuffer& Data)
+	{
+		check(!Data.IsNull());
+		check(Data.GetSize() > 0);
+
+		TArray<UE::DerivedData::FCachePutValueRequest> Requests;
+		Requests.SetNum(1);
+
+		Requests[0].Name = FString("Put-") + RequestId;
+		Requests[0].Key = CreateCacheKey(RequestId);
+		Requests[0].Value = UE::DerivedData::FValue::Compress(Data);
+
+		UE::DerivedData::FRequestOwner BlockingPutOwner(UE::DerivedData::EPriority::Blocking);
+		UE::DerivedData::GetCache().PutValue(Requests, BlockingPutOwner);
+		BlockingPutOwner.Wait();
+	}
+
+	inline FSharedBuffer GetFromDDC(const FString& RequestId)
+	{
+		TArray<UE::DerivedData::FCacheGetValueRequest> Requests;
+		Requests.SetNum(1);
+
+		Requests[0].Name = FString("Get-") + RequestId;
+		Requests[0].Key = CreateCacheKey(RequestId);
+
+		FSharedBuffer Result;
+		UE::DerivedData::FRequestOwner BlockingGetOwner(UE::DerivedData::EPriority::Blocking);
+		UE::DerivedData::GetCache().GetValue(Requests, BlockingGetOwner, [&Result](UE::DerivedData::FCacheGetValueResponse&& Response)
+			{
+				if (Response.Value.HasData() && Response.Value.GetRawSize() > sizeof(uint32))
+				{
+					Result = Response.Value.GetData().Decompress();
+				}
+			});
+		BlockingGetOwner.Wait();
+		return Result;
+	}
+#endif // WITH_EDITOR
+
+	FString GetModelDataIdentifier(const FString& RuntimeName, const FGuid& Guid, const FString& FileIdString, const FString& PlatformName, const FString& Architecture)
+	{
+		return RuntimeName + "-" + Guid.ToString(EGuidFormats::Digits) + "-" + FString::FromInt(UNNERuntimeIREECpu::Version) + "-" + FileIdString + "-" + PlatformName + (!Architecture.IsEmpty() ? ("-" + Architecture) : "");
+	}
+
+	FString GetIntermediateModelDirPath(const FString& PlatformName, const FString& ModelName)
+	{
+		return FPaths::Combine("Intermediate", "Build", PlatformName, UE_PLUGIN_NAME, ModelName);
+	}
+
+	FString GetStagedModelDirPath(const FString& PlatformName)
+	{
+		return FPaths::Combine("Saved", "Cooked", PlatformName, "Engine", "Plugins", UE_PLUGIN_NAME, "Binaries");
+	}
+
+	FString GetPackagedModelDirPath(const FString& PlatformName)
+	{
+		FString PlatformNameShort = PlatformName.Equals("Windows") ? "Win64" : PlatformName;
+		return FPaths::Combine("Binaries", PlatformNameShort, UE_PLUGIN_NAME);
+	}
+} // UE::NNERuntimeIREE::CPU::Private
+
+FGuid UNNERuntimeIREECpu::GUID = FGuid((int32)'I', (int32)'C', (int32)'P', (int32)'U');
+int32 UNNERuntimeIREECpu::Version = 0x00000003;
+
+FString UNNERuntimeIREECpu::GetRuntimeName() const
+{
+	return TEXT("NNERuntimeIREECpu");
+}
+
+bool UNNERuntimeIREECpu::CanCreateModelData(const FString& FileType, TConstArrayView<uint8> FileData, const TMap<FString, TConstArrayView<uint8>>& AdditionalFileData, const FGuid& FileId, const ITargetPlatform* TargetPlatform) const
+{
+#if WITH_EDITOR
+	return 	FileType.Compare(TEXT("mlir"), ESearchCase::IgnoreCase) == 0;
+#else
+	return false;
+#endif // WITH_EDITOR
+}
+
+TSharedPtr<UE::NNE::FSharedModelData> UNNERuntimeIREECpu::CreateModelData(const FString& FileType, TConstArrayView<uint8> FileData, const TMap<FString, TConstArrayView<uint8>>& AdditionalFileData, const FGuid& FileId, const ITargetPlatform* TargetPlatform)
+{
+#if WITH_EDITOR
+	using namespace UE::NNERuntimeIREE::CPU::Private;
+
+	FString TargetPlatformName = TargetPlatform ? TargetPlatform->IniPlatformName() : UGameplayStatics::GetPlatformName();
+	if (!CanCreateModelData(FileType, FileData, AdditionalFileData, FileId, TargetPlatform))
+	{
+		UE_LOG(LogTemp, Warning, TEXT("UNNERuntimeIREECpu cannot create the model data with id %s (Filetype: %s) for platform %s"), *FileId.ToString(EGuidFormats::Digits).ToLower(), *FileType, *TargetPlatformName);
+		return TSharedPtr<UE::NNE::FSharedModelData>();
+	}
+
+	FConfigFile ConfigFile;
+	FString ConfigFilePath;
+	GetUpdatedPlatformConfig(TargetPlatformName, ConfigFile, ConfigFilePath);
+	if (ConfigFile.Dirty)
+	{
+		UE_LOG(LogTemp, Warning, TEXT("UNNERuntimeIREECpu could not find the required settings in config file %s. Please make the file writeable and re-start the editor or manually add the required staging settings or models will not work in packaged builds for platform %s!"), *ConfigFilePath, *TargetPlatformName);
+	}
+
+	TUniquePtr<UE::NNERuntimeIREE::CPU::FCompiler> Compiler = UE::NNERuntimeIREE::CPU::FCompiler::Make(TargetPlatformName);
+	if (!Compiler.IsValid())
+	{
+		UE_LOG(LogTemp, Warning, TEXT("UNNERuntimeIREECpu failed to create a compiler to compile for platform %s"), *TargetPlatformName);
+		return TSharedPtr<UE::NNE::FSharedModelData>();
+	}
+	
+	FString FileIdString = FileId.ToString(EGuidFormats::Digits).ToLower();
+	FString IntermediateDir = FPaths::ConvertRelativePathToFull(FPaths::Combine(FPaths::ProjectDir(), GetIntermediateModelDirPath(TargetPlatformName, FileIdString)));
+	IPlatformFile& PlatformFile = FPlatformFileManager::Get().GetPlatformFile();
+	PlatformFile.DeleteDirectoryRecursively(*IntermediateDir);
+	PlatformFile.CreateDirectoryTree(*IntermediateDir);
+
+	TArray<UE::NNERuntimeIREE::CPU::FCompilerResult> CompilerResults;
+	UNNERuntimeIREEModuleMetaData* CompilerModuleMetaData = NewObject<UNNERuntimeIREEModuleMetaData>();
+	FString StagingDir = FPaths::ConvertRelativePathToFull(FPaths::Combine(FPaths::ProjectDir(), GetStagedModelDirPath(TargetPlatformName)));
+	if (!Compiler->CompileMlir(FileData, FileIdString, IntermediateDir, StagingDir, CompilerResults, CompilerModuleMetaData))
+	{
+		UE_LOG(LogTemp, Warning, TEXT("UNNERuntimeIREECpu failed to compile model %s"), *FileIdString);
+		return TSharedPtr<UE::NNE::FSharedModelData>();
+	}
+
+	for (int32 i = 0; i < CompilerResults.Num(); i++)
+	{
+		TArray<uint8> SharedLibData;
+		FString StagedSharedLibPath = FPaths::Combine(StagingDir, CompilerResults[i].RelativeDirPath, CompilerResults[i].SharedLibraryFileName);
+		if (!FFileHelper::LoadFileToArray(SharedLibData, *StagedSharedLibPath) || SharedLibData.IsEmpty())
+		{
+			UE_LOG(LogTemp, Warning, TEXT("UNNERuntimeIREECpu could not read the shared library \"%s\""), *StagedSharedLibPath);
+			return TSharedPtr<UE::NNE::FSharedModelData>();
+		}
+		FSharedBuffer SharedLibBuffer = MakeSharedBufferFromArray(MoveTemp(SharedLibData));
+		PutIntoDDC(UE::NNERuntimeIREE::CPU::Private::GetModelDataIdentifier(GetRuntimeName(), UNNERuntimeIREECpu::GUID, FileIdString, TargetPlatformName, CompilerResults[i].Architecture) + "-lib", SharedLibBuffer);
+
+		TArray<uint8> VmfbData;
+		FString StagedVmfbPath = FPaths::Combine(StagingDir, CompilerResults[i].RelativeDirPath, CompilerResults[i].VmfbFileName);
+		if (!FFileHelper::LoadFileToArray(VmfbData, *StagedVmfbPath) || VmfbData.IsEmpty())
+		{
+			UE_LOG(LogTemp, Warning, TEXT("UNNERuntimeIREECpu could not read the vmfb data \"%s\""), *StagedVmfbPath);
+			return TSharedPtr<UE::NNE::FSharedModelData>();
+		}
+		FSharedBuffer VmfbBuffer = MakeSharedBufferFromArray(MoveTemp(VmfbData));
+		PutIntoDDC(UE::NNERuntimeIREE::CPU::Private::GetModelDataIdentifier(GetRuntimeName(), UNNERuntimeIREECpu::GUID, FileIdString, TargetPlatformName, CompilerResults[i].Architecture) + "-vmfb", VmfbBuffer);
+	}
+
+	TArray<uint8> ResultData;
+	FMemoryWriter Writer(ResultData);
+	Writer << UNNERuntimeIREECpu::GUID;
+	Writer << UNNERuntimeIREECpu::Version;
+	FGuid FileIdCopy = FileId;
+	Writer << FileIdCopy;
+
+	TArray<uint8> ModuleMetaData;
+	if (AdditionalFileData.Contains("IREEModuleMetaData"))
+	{
+		ModuleMetaData = AdditionalFileData["IREEModuleMetaData"];
+	}
+	if (ModuleMetaData.IsEmpty())
+	{
+		FMemoryWriter ObjectWriter(ModuleMetaData);
+		CompilerModuleMetaData->Serialize(ObjectWriter);
+	}
+	Writer << ModuleMetaData;
+
+	int32 NumArchitectures = CompilerResults.Num();
+	Writer << NumArchitectures;
+	for (int32 i = 0; i < NumArchitectures; i++)
+	{
+		Writer << CompilerResults[i].Architecture;
+		Writer << CompilerResults[i].RelativeDirPath;
+		Writer << CompilerResults[i].SharedLibraryFileName;
+		Writer << CompilerResults[i].VmfbFileName;
+		Writer << CompilerResults[i].SharedLibraryEntryPointName;
+	}
+
+	return MakeShared<UE::NNE::FSharedModelData>(MakeSharedBufferFromArray(MoveTemp(ResultData)), 0);
+#else
+	return TSharedPtr<UE::NNE::FSharedModelData>();
+#endif // WITH_EDITOR
+}
+
+FString UNNERuntimeIREECpu::GetModelDataIdentifier(const FString& FileType, TConstArrayView<uint8> FileData, const TMap<FString, TConstArrayView<uint8>>& AdditionalFileData, const FGuid& FileId, const ITargetPlatform* TargetPlatform) const
+{
+	// Leave architecture blank as there is only one model data for all architectures of a given platform, only the vmfb and shared lib are different
+	FString PlatformName = TargetPlatform ? TargetPlatform->IniPlatformName() : UGameplayStatics::GetPlatformName();
+	return UE::NNERuntimeIREE::CPU::Private::GetModelDataIdentifier(GetRuntimeName(), UNNERuntimeIREECpu::GUID, FileId.ToString(EGuidFormats::Digits), PlatformName, "");
+}
+
+bool UNNERuntimeIREECpu::CanCreateModelCPU(TObjectPtr<UNNEModelData> ModelData) const
+{
+	check(ModelData != nullptr);
+
+	TSharedPtr<UE::NNE::FSharedModelData> SharedData = ModelData->GetModelData(GetRuntimeName());
+	if (!SharedData.IsValid())
+	{
+		return false;
+	}
+
+	TConstArrayView<uint8> SharedDataView = SharedData->GetView();
+	int32 GuidSize = sizeof(UNNERuntimeIREECpu::GUID);
+	int32 VersionSize = sizeof(UNNERuntimeIREECpu::Version);
+	if (SharedDataView.Num() <= GuidSize + VersionSize)
+	{
+		return false;
+	}
+
+	bool bResult = FGenericPlatformMemory::Memcmp(&(SharedDataView[0]), &(UNNERuntimeIREECpu::GUID), GuidSize) == 0;
+	bResult &= FGenericPlatformMemory::Memcmp(&(SharedDataView[GuidSize]), &(UNNERuntimeIREECpu::Version), VersionSize) == 0;
+	return bResult;
+}
+
+TSharedPtr<UE::NNE::IModelCPU> UNNERuntimeIREECpu::CreateModelCPU(TObjectPtr<UNNEModelData> ModelData)
+{
+	check(ModelData != nullptr);
+
+	using namespace UE::NNERuntimeIREE::CPU::Private;
+
+	if (!CanCreateModelCPU(ModelData))
+	{
+		UE_LOG(LogTemp, Warning, TEXT("UNNERuntimeIREECpu cannot create a model from the model data with id %s"), *ModelData->GetFileId().ToString(EGuidFormats::Digits));
+		return TSharedPtr<UE::NNE::IModelCPU>();
+	}
+
+	FString CurrentArchitecture = "";
+#if PLATFORM_CPU_X86_FAMILY
+	CurrentArchitecture = "x86_64";
+#elif PLATFORM_CPU_ARM_FAMILY
+	CurrentArchitecture = "arm64";
+#endif
+
+	TSharedPtr<UE::NNE::FSharedModelData> SharedData = ModelData->GetModelData(GetRuntimeName());
+	check(SharedData.IsValid());
+	TConstArrayView<uint8> SharedDataView = SharedData->GetView();
+	FMemoryReaderView Reader(SharedDataView);
+	FGuid DataGuid = FGuid();
+	Reader << DataGuid;
+	int32 VersionVersion = 0;
+	Reader << VersionVersion;
+	FGuid FileId = FGuid();
+	Reader << FileId;
+
+	TArray<uint8> ModuleDataArray;
+	Reader << ModuleDataArray;
+	if (ModuleDataArray.IsEmpty())
+	{
+		UE_LOG(LogTemp, Warning, TEXT("UNNERuntimeIREECpu failed to find any module meta data, please reimport the original model"));
+		return TSharedPtr<UE::NNE::IModelCPU>();
+	}
+	UNNERuntimeIREEModuleMetaData* ModuleMetaData = NewObject<UNNERuntimeIREEModuleMetaData>();
+	FMemoryReaderView ObjectReader(ModuleDataArray);
+	ModuleMetaData->Serialize(ObjectReader);
+	if (ModuleMetaData->FunctionMetaData.IsEmpty())
+	{
+		UE_LOG(LogTemp, Warning, TEXT("UNNERuntimeIREECpu failed to parse the module meta data, please reimport the original model"));
+		return TSharedPtr<UE::NNE::IModelCPU>();
+	}
+
+	int32 NumArchitectures = 0;
+	Reader << NumArchitectures;
+
+	bool bFound = false;
+	FString Architecture = "";
+	FString RelativeDirPath = "";
+	FString SharedLibraryFileName = "";
+	FString VmfbFileName = "";
+	FString SharedLibraryEntryPointName = "";
+	for (int32 i = 0; i < NumArchitectures; i++)
+	{
+		FString TmpArchitecture = "";
+		Reader << TmpArchitecture;
+		FString TmpRelativeDirPath = "";
+		Reader << TmpRelativeDirPath;
+		FString TmpSharedLibraryFileName = "";
+		Reader << TmpSharedLibraryFileName;
+		FString TmpVmfbFileName = "";
+		Reader << TmpVmfbFileName;
+		FString TmpSharedLibraryEntryPointName = "";
+		Reader << TmpSharedLibraryEntryPointName;
+
+		if (TmpArchitecture.IsEmpty() && !bFound)
+		{
+			Architecture = TmpArchitecture;
+			RelativeDirPath = TmpRelativeDirPath;
+			SharedLibraryFileName = TmpSharedLibraryFileName;
+			VmfbFileName = TmpVmfbFileName;
+			SharedLibraryEntryPointName = TmpSharedLibraryEntryPointName;
+			bFound = true;
+		}
+		else if (TmpArchitecture.Equals(CurrentArchitecture))
+		{
+			Architecture = TmpArchitecture;
+			RelativeDirPath = TmpRelativeDirPath;
+			SharedLibraryFileName = TmpSharedLibraryFileName;
+			VmfbFileName = TmpVmfbFileName;
+			SharedLibraryEntryPointName = TmpSharedLibraryEntryPointName;
+			bFound = true;
+		}
+	}
+	if (!bFound)
+	{
+		UE_LOG(LogTemp, Warning, TEXT("UNNERuntimeIREECpu failed to find a matching architecture for \'%s\'"), *CurrentArchitecture);
+		return TSharedPtr<UE::NNE::IModelCPU>();
+	}
+
+	FString FileIdString = FileId.ToString(EGuidFormats::Digits).ToLower();
+#if WITH_EDITOR
+	FString SharedLibraryDirPath = FPaths::ConvertRelativePathToFull(FPaths::Combine(FPaths::ProjectDir(), GetIntermediateModelDirPath(UGameplayStatics::GetPlatformName(), FileIdString), RelativeDirPath));
+#else
+	FString SharedLibraryDirPath = FPaths::ConvertRelativePathToFull(FPaths::Combine(FPaths::ProjectDir(), GetPackagedModelDirPath(UGameplayStatics::GetPlatformName()), RelativeDirPath));
+#endif // WITH_EDITOR
+
+#if WITH_EDITOR
+	IPlatformFile& PlatformFile = FPlatformFileManager::Get().GetPlatformFile();
+	FString SharedLibraryFilePath = FPaths::Combine(SharedLibraryDirPath, SharedLibraryFileName);
+	if (!PlatformFile.FileExists(*SharedLibraryFilePath))
+	{
+		FSharedBuffer SharedBuffer = GetFromDDC(UE::NNERuntimeIREE::CPU::Private::GetModelDataIdentifier(GetRuntimeName(), UNNERuntimeIREECpu::GUID, FileIdString, UGameplayStatics::GetPlatformName(), Architecture) + "-lib");
+		if (SharedBuffer.GetSize() <= 0)
+		{
+			UE_LOG(LogTemp, Warning, TEXT("UNNERuntimeIREECpu could not fetch the shared library %s from DDC"), *SharedLibraryFileName);
+			return TSharedPtr<UE::NNE::IModelCPU>();
+		}
+		PlatformFile.CreateDirectoryTree(*SharedLibraryDirPath);
+		FFileHelper::SaveArrayToFile(TConstArrayView<uint8>((uint8*)SharedBuffer.GetData(), SharedBuffer.GetSize()), *SharedLibraryFilePath);
+	}
+
+	FString VmfbFilePath = FPaths::Combine(SharedLibraryDirPath, VmfbFileName);
+	if (!PlatformFile.FileExists(*VmfbFilePath))
+	{
+		FSharedBuffer SharedBuffer = GetFromDDC(UE::NNERuntimeIREE::CPU::Private::GetModelDataIdentifier(GetRuntimeName(), UNNERuntimeIREECpu::GUID, FileIdString, UGameplayStatics::GetPlatformName(), Architecture) + "-vmfb");
+		if (SharedBuffer.GetSize() <= 0)
+		{
+			UE_LOG(LogTemp, Warning, TEXT("UNNERuntimeIREECpu could not fetch the vmfb %s from DDC"), *VmfbFileName);
+			return TSharedPtr<UE::NNE::IModelCPU>();
+		}
+		PlatformFile.CreateDirectoryTree(*SharedLibraryDirPath);
+		FFileHelper::SaveArrayToFile(TConstArrayView<uint8>((uint8*)SharedBuffer.GetData(), SharedBuffer.GetSize()), *VmfbFilePath);
+	}
+#endif // WITH_EDITOR
+
+	TSharedPtr<UE::NNE::IModelCPU> Model = UE::NNERuntimeIREE::CPU::FModel::Make(SharedLibraryDirPath, SharedLibraryFileName, VmfbFileName, SharedLibraryEntryPointName, *ModuleMetaData);
+	if (!Model.IsValid())
+	{
+		UE_LOG(LogTemp, Warning, TEXT("UNNERuntimeIREECpu could not initialize the model created from model data with id %s"), *FileIdString);
+		return TSharedPtr<UE::NNE::IModelCPU>();
+	}
+
+	if (FEngineAnalytics::IsAvailable())
+	{
+		TArray<FAnalyticsEventAttribute> Attributes = MakeAnalyticsEventAttributeArray(
+			TEXT("PlatformName"), UGameplayStatics::GetPlatformName(),
+			TEXT("HashedRuntimeName"), FMD5::HashAnsiString(*GetRuntimeName()),
+			TEXT("ModelDataSize"), SharedDataView.Num()
+		);
+		FEngineAnalytics::GetProvider().RecordEvent(TEXT("NeuralNetworkEngine.CreateModel"), Attributes);
+	}
+
+	return Model;
+}
+
+void UNNERuntimeIREECpu::GetUpdatedPlatformConfig(const FString& PlatformName, FConfigFile& ConfigFile, FString& ConfigFilePath)
+{ 
+	FString ConfigFolderPath = FPaths::ConvertRelativePathToFull(FPaths::ProjectConfigDir());
+	ConfigFilePath = FPaths::Combine(ConfigFolderPath, PlatformName, PlatformName + "Game.ini");
+
+	ConfigFile.Read(ConfigFilePath);
+
+	FString StagingPath = FString("/") + UE::NNERuntimeIREE::CPU::Private::GetStagedModelDirPath(PlatformName);
+	FString PackagingPath = FString("/") + UE::NNERuntimeIREE::CPU::Private::GetPackagedModelDirPath(PlatformName);
+
+	ConfigFile.AddUniqueToSection(TEXT("/Script/UnrealEd.ProjectPackagingSettings"), TEXT("+DirectoriesToAlwaysStageAsNonUFS"), FString("(Path=\"..") + StagingPath + FString("\")"));
+	ConfigFile.AddUniqueToSection(TEXT("Staging"), TEXT("+RemapDirectories"), FString("(From=\"") + FApp::GetProjectName() + StagingPath + FString("\", To=\"") + FApp::GetProjectName() + PackagingPath + FString("\")"));
+	ConfigFile.AddUniqueToSection(TEXT("Staging"), TEXT("+AllowedDirectories"), FApp::GetProjectName() + PackagingPath);
+}
+
+FString UNNERuntimeIREEGpu::GetRuntimeName() const
+{
+	return TEXT("");
+}
+
+bool UNNERuntimeIREEGpu::CanCreateModelData(const FString& FileType, TConstArrayView<uint8> FileData, const TMap<FString, TConstArrayView<uint8>>& AdditionalFileData, const FGuid& FileId, const ITargetPlatform* TargetPlatform) const
+{
+#if WITH_EDITOR
+	return	FileType.Compare(TEXT("mlir"), ESearchCase::IgnoreCase) == 0;
+#else
+	return false;
+#endif // WITH_EDITOR
+}
+
+TSharedPtr<UE::NNE::FSharedModelData> UNNERuntimeIREEGpu::CreateModelData(const FString& FileType, TConstArrayView<uint8> FileData, const TMap<FString, TConstArrayView<uint8>>& AdditionalFileData, const FGuid& FileId, const ITargetPlatform* TargetPlatform)
+{
+	return TSharedPtr<UE::NNE::FSharedModelData>();
+}
+
+FString UNNERuntimeIREEGpu::GetModelDataIdentifier(const FString& FileType, TConstArrayView<uint8> FileData, const TMap<FString, TConstArrayView<uint8>>& AdditionalFileData, const FGuid& FileId, const ITargetPlatform* TargetPlatform) const
+{
+	FString PlatformName = TargetPlatform ? TargetPlatform->IniPlatformName() : UGameplayStatics::GetPlatformName();
+	return UE::NNERuntimeIREE::CPU::Private::GetModelDataIdentifier(GetRuntimeName(), GetGUID(), FileId.ToString(EGuidFormats::Digits), PlatformName, "");
+}
+
+bool UNNERuntimeIREEGpu::CanCreateModelGPU(TObjectPtr<UNNEModelData> ModelData) const
+{
+	check(ModelData != nullptr);
+
+	TSharedPtr<UE::NNE::FSharedModelData> SharedData = ModelData->GetModelData(GetRuntimeName());
+	if (!SharedData.IsValid())
+	{
+		return false;
+	}
+
+	TConstArrayView<uint8> SharedDataView = SharedData->GetView();
+	FGuid Guid = GetGUID();
+	int32 Version = GetVersion();
+	int32 GuidSize = sizeof(Guid);
+	int32 VersionSize = sizeof(Version);
+	if (SharedDataView.Num() <= GuidSize + VersionSize)
+	{
+		return false;
+	}
+
+	bool bResult = FGenericPlatformMemory::Memcmp(&(SharedDataView[0]), &(Guid), GuidSize) == 0;
+	bResult &= FGenericPlatformMemory::Memcmp(&(SharedDataView[GuidSize]), &(Version), VersionSize) == 0;
+	return bResult;
+}
+
+TSharedPtr<UE::NNE::IModelGPU> UNNERuntimeIREEGpu::CreateModelGPU(TObjectPtr<UNNEModelData> ModelData)
+{
+	check(ModelData != nullptr);
+
+	if (!CanCreateModelGPU(ModelData))
+	{
+		return TSharedPtr<UE::NNE::IModelGPU>();
+	}
+
+	check(ModelData->GetModelData(GetRuntimeName()).IsValid());
+
+	UE::NNE::IModelGPU* IModel = nullptr;
+	TConstArrayView<uint8> SharedDataView = ModelData->GetModelData(GetRuntimeName())->GetView();
+
+	if (FEngineAnalytics::IsAvailable())
+	{
+		TArray<FAnalyticsEventAttribute> Attributes = MakeAnalyticsEventAttributeArray(
+			TEXT("PlatformName"), UGameplayStatics::GetPlatformName(),
+			TEXT("HashedRuntimeName"), FMD5::HashAnsiString(*GetRuntimeName()),
+			TEXT("ModelDataSize"), SharedDataView.Num()
+		);
+		FEngineAnalytics::GetProvider().RecordEvent(TEXT("NeuralNetworkEngine.CreateModel"), Attributes);
+	}
+
+	return TSharedPtr<UE::NNE::IModelGPU>(IModel);
+}
+
+bool UNNERuntimeIREEGpu::IsAvailable() const
+{
+	return false;
+}
+
+FGuid UNNERuntimeIREEGpu::GetGUID() const
+{
+	return FGuid();
+}
+
+int32 UNNERuntimeIREEGpu::GetVersion() const
+{
+	return 0;
+}
+
+FGuid UNNERuntimeIREECuda::GUID = FGuid((int32)'I', (int32)'G', (int32)'C', (int32)'U');
+int32 UNNERuntimeIREECuda::Version = 0x00000001;
+
+FString UNNERuntimeIREECuda::GetRuntimeName() const
+{
+	return TEXT("NNERuntimeIREECuda");
+}
+
+bool UNNERuntimeIREECuda::IsAvailable() const
+{
+	return false;
+}
+
+FGuid UNNERuntimeIREECuda::GetGUID() const
+{
+	return GUID;
+}
+
+int32 UNNERuntimeIREECuda::GetVersion() const
+{
+	return Version;
+}
+
+FGuid UNNERuntimeIREEVulkan::GUID = FGuid((int32)'I', (int32)'G', (int32)'V', (int32)'U');
+int32 UNNERuntimeIREEVulkan::Version = 0x00000001;
+
+FString UNNERuntimeIREEVulkan::GetRuntimeName() const
+{
+	return TEXT("NNERuntimeIREEVulkan");
+}
+
+bool UNNERuntimeIREEVulkan::IsAvailable() const
+{
+	return false;
+}
+
+FGuid UNNERuntimeIREEVulkan::GetGUID() const
+{
+	return GUID;
+}
+
+int32 UNNERuntimeIREEVulkan::GetVersion() const
+{
+	return Version;
+}
+
+FGuid UNNERuntimeIREERdg::GUID = FGuid((int32)'I', (int32)'R', (int32)'D', (int32)'G');
+int32 UNNERuntimeIREERdg::Version = 0x00000001;
+
+FString UNNERuntimeIREERdg::GetRuntimeName() const
+{
+	return TEXT("NNERuntimeIREERdg");
+}
+
+bool UNNERuntimeIREERdg::CanCreateModelData(const FString& FileType, TConstArrayView<uint8> FileData, const TMap<FString, TConstArrayView<uint8>>& AdditionalFileData, const FGuid& FileId, const ITargetPlatform* TargetPlatform) const
+{
+#if WITH_EDITOR
+	return	FileType.Compare(TEXT("mlir"), ESearchCase::IgnoreCase) == 0;
+#else
+	return false;
+#endif // WITH_EDITOR
+}
+
+TSharedPtr<UE::NNE::FSharedModelData> UNNERuntimeIREERdg::CreateModelData(const FString& FileType, TConstArrayView<uint8> FileData, const TMap<FString, TConstArrayView<uint8>>& AdditionalFileData, const FGuid& FileId, const ITargetPlatform* TargetPlatform)
+{
+	return TSharedPtr<UE::NNE::FSharedModelData>();
+}
+
+FString UNNERuntimeIREERdg::GetModelDataIdentifier(const FString& FileType, TConstArrayView<uint8> FileData, const TMap<FString, TConstArrayView<uint8>>& AdditionalFileData, const FGuid& FileId, const ITargetPlatform* TargetPlatform) const
+{
+	FString PlatformName = TargetPlatform ? TargetPlatform->IniPlatformName() : UGameplayStatics::GetPlatformName();
+	return UE::NNERuntimeIREE::CPU::Private::GetModelDataIdentifier(GetRuntimeName(), UNNERuntimeIREERdg::GUID, FileId.ToString(EGuidFormats::Digits), PlatformName, "");
+}
+
+bool UNNERuntimeIREERdg::CanCreateModelRDG(TObjectPtr<UNNEModelData> ModelData) const
+{
+	check(ModelData != nullptr);
+
+	TSharedPtr<UE::NNE::FSharedModelData> SharedData = ModelData->GetModelData(GetRuntimeName());
+	if (!SharedData.IsValid())
+	{
+		return false;
+	}
+
+	TConstArrayView<uint8> SharedDataView = SharedData->GetView();
+	int32 GuidSize = sizeof(UNNERuntimeIREERdg::GUID);
+	int32 VersionSize = sizeof(UNNERuntimeIREERdg::Version);
+	if (SharedDataView.Num() <= GuidSize + VersionSize)
+	{
+		return false;
+	}
+
+	bool bResult = FGenericPlatformMemory::Memcmp(&(SharedDataView[0]), &(UNNERuntimeIREERdg::GUID), GuidSize) == 0;
+	bResult &= FGenericPlatformMemory::Memcmp(&(SharedDataView[GuidSize]), &(UNNERuntimeIREERdg::Version), VersionSize) == 0;
+	return bResult;
+}
+
+TSharedPtr<UE::NNE::IModelRDG> UNNERuntimeIREERdg::CreateModelRDG(TObjectPtr<UNNEModelData> ModelData)
+{
+	check(ModelData != nullptr);
+
+	if (!CanCreateModelRDG(ModelData))
+	{
+		return TSharedPtr<UE::NNE::IModelRDG>();
+	}
+
+	check(ModelData->GetModelData(GetRuntimeName()).IsValid());
+
+	UE::NNE::IModelRDG* IModel = nullptr;
+	TConstArrayView<uint8> SharedDataView = ModelData->GetModelData(GetRuntimeName())->GetView();
+
+	if (FEngineAnalytics::IsAvailable())
+	{
+		TArray<FAnalyticsEventAttribute> Attributes = MakeAnalyticsEventAttributeArray(
+			TEXT("PlatformName"), UGameplayStatics::GetPlatformName(),
+			TEXT("HashedRuntimeName"), FMD5::HashAnsiString(*GetRuntimeName()),
+			TEXT("ModelDataSize"), SharedDataView.Num()
+		);
+		FEngineAnalytics::GetProvider().RecordEvent(TEXT("NeuralNetworkEngine.CreateModel"), Attributes);
+	}
+
+	return TSharedPtr<UE::NNE::IModelRDG>(IModel);
+}
+
+bool UNNERuntimeIREERdg::IsAvailable() const
+{
+	return false;
+}
+
+#endif // WITH_NNE_RUNTIME_IREE
\ No newline at end of file
diff --git a/Engine/Plugins/Experimental/NNERuntimeIREE/Source/NNERuntimeIREE/Private/NNERuntimeIREE.h b/Engine/Plugins/Experimental/NNERuntimeIREE/Source/NNERuntimeIREE/Private/NNERuntimeIREE.h
new file mode 100644
--- /dev/null
+++ b/Engine/Plugins/Experimental/NNERuntimeIREE/Source/NNERuntimeIREE/Private/NNERuntimeIREE.h
@@ -0,0 +1,170 @@
+// Copyright Epic Games, Inc. All Rights Reserved.
+
+#pragma once
+
+#include "CoreMinimal.h"
+#include "Misc/ConfigCacheIni.h"
+#include "Misc/Guid.h"
+#include "NNERuntime.h"
+#include "NNERuntimeCPU.h"
+#include "NNERuntimeGPU.h"
+#include "NNERuntimeRDG.h"
+
+#include "NNERuntimeIREE.generated.h"
+
+UCLASS()
+class UNNERuntimeIREECpu : public UObject, public INNERuntime, public INNERuntimeCPU
+{
+	GENERATED_BODY()
+
+public:
+#ifdef WITH_NNE_RUNTIME_IREE
+	static FGuid GUID;
+	static int32 Version;
+
+	//~ Begin INNERuntime Interface
+	virtual FString GetRuntimeName() const override;
+	virtual bool CanCreateModelData(const FString& FileType, TConstArrayView<uint8> FileData, const TMap<FString, TConstArrayView<uint8>>& AdditionalFileData, const FGuid& FileId, const ITargetPlatform* TargetPlatform) const override;
+	virtual TSharedPtr<UE::NNE::FSharedModelData> CreateModelData(const FString& FileType, TConstArrayView<uint8> FileData, const TMap<FString, TConstArrayView<uint8>>& AdditionalFileData, const FGuid& FileId, const ITargetPlatform* TargetPlatform) override;
+	virtual FString GetModelDataIdentifier(const FString& FileType, TConstArrayView<uint8> FileData, const TMap<FString, TConstArrayView<uint8>>& AdditionalFileData, const FGuid& FileId, const ITargetPlatform* TargetPlatform) const override;
+	//~ End INNERuntime Interface
+
+	//~ Begin INNERuntimeCPU Interface
+	virtual bool CanCreateModelCPU(TObjectPtr<UNNEModelData> ModelData) const override;
+	virtual TSharedPtr<UE::NNE::IModelCPU> CreateModelCPU(TObjectPtr<UNNEModelData> ModelData) override;
+	//~ End INNERuntimeCPU Interface
+
+	static void GetUpdatedPlatformConfig(const FString& PlatformName, FConfigFile& ConfigFile, FString& ConfigFilePath);
+#else
+	//~ Begin INNERuntime Interface
+	virtual FString GetRuntimeName() const override { return ""; };
+	virtual bool CanCreateModelData(const FString& FileType, TConstArrayView<uint8> FileData, const TMap<FString, TConstArrayView<uint8>>& AdditionalFileData, const FGuid& FileId, const ITargetPlatform* TargetPlatform) const override { return false; };
+	virtual TSharedPtr<UE::NNE::FSharedModelData> CreateModelData(const FString& FileType, TConstArrayView<uint8> FileData, const TMap<FString, TConstArrayView<uint8>>& AdditionalFileData, const FGuid& FileId, const ITargetPlatform* TargetPlatform) override { return TSharedPtr<UE::NNE::FSharedModelData>(); };
+	virtual FString GetModelDataIdentifier(const FString& FileType, TConstArrayView<uint8> FileData, const TMap<FString, TConstArrayView<uint8>>& AdditionalFileData, const FGuid& FileId, const ITargetPlatform* TargetPlatform) const override { return ""; };
+	//~ End INNERuntime Interface
+
+	//~ Begin INNERuntimeCPU Interface
+	virtual bool CanCreateModelCPU(TObjectPtr<UNNEModelData> ModelData) const override { return false; };
+	virtual TSharedPtr<UE::NNE::IModelCPU> CreateModelCPU(TObjectPtr<UNNEModelData> ModelData) override { return TSharedPtr<UE::NNE::IModelCPU>(); };
+	//~ End INNERuntimeCPU Interface
+#endif // WITH_NNE_RUNTIME_IREE
+};
+
+UCLASS()
+class UNNERuntimeIREEGpu : public UObject, public INNERuntime, public INNERuntimeGPU
+{
+	GENERATED_BODY()
+
+public:
+#ifdef WITH_NNE_RUNTIME_IREE
+	//~ Begin INNERuntime Interface
+	virtual FString GetRuntimeName() const override;
+	virtual bool CanCreateModelData(const FString& FileType, TConstArrayView<uint8> FileData, const TMap<FString, TConstArrayView<uint8>>& AdditionalFileData, const FGuid& FileId, const ITargetPlatform* TargetPlatform) const override;
+	virtual TSharedPtr<UE::NNE::FSharedModelData> CreateModelData(const FString& FileType, TConstArrayView<uint8> FileData, const TMap<FString, TConstArrayView<uint8>>& AdditionalFileData, const FGuid& FileId, const ITargetPlatform* TargetPlatform) override;
+	virtual FString GetModelDataIdentifier(const FString& FileType, TConstArrayView<uint8> FileData, const TMap<FString, TConstArrayView<uint8>>& AdditionalFileData, const FGuid& FileId, const ITargetPlatform* TargetPlatform) const override;
+	//~ End INNERuntime Interface
+
+	//~ Begin INNERuntimeGPU Interface
+	virtual bool CanCreateModelGPU(TObjectPtr<UNNEModelData> ModelData) const override;
+	virtual TSharedPtr<UE::NNE::IModelGPU> CreateModelGPU(TObjectPtr<UNNEModelData> ModelData) override;
+	//~ End INNERuntimeGPU Interface
+
+	virtual bool IsAvailable() const;
+	virtual FGuid GetGUID() const;
+	virtual int32 GetVersion() const;
+#else
+	//~ Begin INNERuntime Interface
+	virtual FString GetRuntimeName() const override { return ""; };
+	virtual bool CanCreateModelData(const FString& FileType, TConstArrayView<uint8> FileData, const TMap<FString, TConstArrayView<uint8>>& AdditionalFileData, const FGuid& FileId, const ITargetPlatform* TargetPlatform) const override { return false; };
+	virtual TSharedPtr<UE::NNE::FSharedModelData> CreateModelData(const FString& FileType, TConstArrayView<uint8> FileData, const TMap<FString, TConstArrayView<uint8>>& AdditionalFileData, const FGuid& FileId, const ITargetPlatform* TargetPlatform) override { return TSharedPtr<UE::NNE::FSharedModelData>(); };
+	virtual FString GetModelDataIdentifier(const FString& FileType, TConstArrayView<uint8> FileData, const TMap<FString, TConstArrayView<uint8>>& AdditionalFileData, const FGuid& FileId, const ITargetPlatform* TargetPlatform) const override { return ""; };
+	//~ End INNERuntime Interface
+
+	//~ Begin INNERuntimeGPU Interface
+	virtual bool CanCreateModelGPU(TObjectPtr<UNNEModelData> ModelData) const override { return false; };
+	virtual TSharedPtr<UE::NNE::IModelGPU> CreateModelGPU(TObjectPtr<UNNEModelData> ModelData) override { return TSharedPtr<UE::NNE::IModelGPU>(); };
+	//~ End INNERuntimeGPU Interface
+#endif // WITH_NNE_RUNTIME_IREE
+};
+
+UCLASS()
+class UNNERuntimeIREECuda : public UNNERuntimeIREEGpu
+{
+	GENERATED_BODY()
+
+public:
+#ifdef WITH_NNE_RUNTIME_IREE
+	static FGuid GUID;
+	static int32 Version;
+
+	//~ Begin INNERuntime Interface
+	virtual FString GetRuntimeName() const override;
+	//~ End INNERuntime Interface
+
+	//~ Begin UNNERuntimeIREEGpu Interface
+	virtual bool IsAvailable() const override;
+	virtual FGuid GetGUID() const override;
+	virtual int32 GetVersion() const override;
+	//~ Begin UNNERuntimeIREEGpu Interface
+#endif // WITH_NNE_RUNTIME_IREE
+};
+
+UCLASS()
+class UNNERuntimeIREEVulkan : public UNNERuntimeIREEGpu
+{
+	GENERATED_BODY()
+
+public:
+#ifdef WITH_NNE_RUNTIME_IREE
+	static FGuid GUID;
+	static int32 Version;
+
+	//~ Begin INNERuntime Interface
+	virtual FString GetRuntimeName() const override;
+	//~ End INNERuntime Interface
+
+	//~ Begin UNNERuntimeIREEGpu Interface
+	virtual bool IsAvailable() const override;
+	virtual FGuid GetGUID() const override;
+	virtual int32 GetVersion() const override;
+	//~ Begin UNNERuntimeIREEGpu Interface
+#endif // WITH_NNE_RUNTIME_IREE
+};
+
+UCLASS()
+class UNNERuntimeIREERdg : public UObject, public INNERuntime, public INNERuntimeRDG
+{
+	GENERATED_BODY()
+
+public:
+#ifdef WITH_NNE_RUNTIME_IREE
+	static FGuid GUID;
+	static int32 Version;
+
+	//~ Begin INNERuntime Interface
+	virtual FString GetRuntimeName() const override;
+	virtual bool CanCreateModelData(const FString& FileType, TConstArrayView<uint8> FileData, const TMap<FString, TConstArrayView<uint8>>& AdditionalFileData, const FGuid& FileId, const ITargetPlatform* TargetPlatform) const override;
+	virtual TSharedPtr<UE::NNE::FSharedModelData> CreateModelData(const FString& FileType, TConstArrayView<uint8> FileData, const TMap<FString, TConstArrayView<uint8>>& AdditionalFileData, const FGuid& FileId, const ITargetPlatform* TargetPlatform) override;
+	virtual FString GetModelDataIdentifier(const FString& FileType, TConstArrayView<uint8> FileData, const TMap<FString, TConstArrayView<uint8>>& AdditionalFileData, const FGuid& FileId, const ITargetPlatform* TargetPlatform) const override;
+	//~ End INNERuntime Interface
+
+	//~ Begin INNERuntimeRdg Interface
+	virtual bool CanCreateModelRDG(TObjectPtr<UNNEModelData> ModelData) const override;
+	virtual TSharedPtr<UE::NNE::IModelRDG> CreateModelRDG(TObjectPtr<UNNEModelData> ModelData) override;
+	//~ End INNERuntimeRdg Interface
+
+	bool IsAvailable() const;
+#else
+	//~ Begin INNERuntime Interface
+	virtual FString GetRuntimeName() const override { return ""; };
+	virtual bool CanCreateModelData(const FString& FileType, TConstArrayView<uint8> FileData, const TMap<FString, TConstArrayView<uint8>>& AdditionalFileData, const FGuid& FileId, const ITargetPlatform* TargetPlatform) const override { return false; };
+	virtual TSharedPtr<UE::NNE::FSharedModelData> CreateModelData(const FString& FileType, TConstArrayView<uint8> FileData, const TMap<FString, TConstArrayView<uint8>>& AdditionalFileData, const FGuid& FileId, const ITargetPlatform* TargetPlatform) override { return TSharedPtr<UE::NNE::FSharedModelData>(); };
+	virtual FString GetModelDataIdentifier(const FString& FileType, TConstArrayView<uint8> FileData, const TMap<FString, TConstArrayView<uint8>>& AdditionalFileData, const FGuid& FileId, const ITargetPlatform* TargetPlatform) const override { return ""; };
+	//~ End INNERuntime Interface
+
+	//~ Begin INNERuntimeRDG Interface
+	virtual bool CanCreateModelRDG(TObjectPtr<UNNEModelData> ModelData) const override { return false; };
+	virtual TSharedPtr<UE::NNE::IModelRDG> CreateModelRDG(TObjectPtr<UNNEModelData> ModelData) override { return TSharedPtr<UE::NNE::IModelRDG>(); };
+	//~ End INNERuntimeRDG Interface
+#endif // WITH_NNE_RUNTIME_IREE
+};
\ No newline at end of file
diff --git a/Engine/Plugins/Experimental/NNERuntimeIREE/Source/NNERuntimeIREE/Private/NNERuntimeIREECommon.cpp b/Engine/Plugins/Experimental/NNERuntimeIREE/Source/NNERuntimeIREE/Private/NNERuntimeIREECommon.cpp
deleted file mode 100644
diff --git a/Engine/Plugins/Experimental/NNERuntimeIREE/Source/NNERuntimeIREE/Private/NNERuntimeIREECommon.h b/Engine/Plugins/Experimental/NNERuntimeIREE/Source/NNERuntimeIREE/Private/NNERuntimeIREECommon.h
deleted file mode 100644
diff --git a/Engine/Plugins/Experimental/NNERuntimeIREE/Source/NNERuntimeIREE/Private/NNERuntimeIREECompiler.cpp b/Engine/Plugins/Experimental/NNERuntimeIREE/Source/NNERuntimeIREE/Private/NNERuntimeIREECompiler.cpp
new file mode 100644
--- /dev/null
+++ b/Engine/Plugins/Experimental/NNERuntimeIREE/Source/NNERuntimeIREE/Private/NNERuntimeIREECompiler.cpp
@@ -0,0 +1,326 @@
+// Copyright Epic Games, Inc. All Rights Reserved.
+
+#include "NNERuntimeIREECompiler.h"
+
+#ifdef WITH_NNE_RUNTIME_IREE
+#if WITH_EDITOR
+
+#include "GenericPlatform/GenericPlatformFile.h"
+#include "GenericPlatform/GenericPlatformMisc.h"
+#include "GenericPlatform/GenericPlatformProcess.h"
+#include "HAL/FileManager.h"
+#include "HAL/PlatformFileManager.h"
+#include "Interfaces/IPluginManager.h"
+#include "Kismet/GameplayStatics.h"
+#include "Misc/FileHelper.h"
+#include "Misc/Paths.h"
+#include "Serialization/JsonReader.h"
+#include "Serialization/JsonSerializer.h"
+#include "Templates/SharedPointer.h"
+
+namespace UE::NNERuntimeIREE
+{
+	namespace CPU
+	{
+		namespace Private
+		{
+			bool ResolveEnvironmentVariables(FString& String)
+			{
+				const FString StartString = "$ENV{";
+				const FString EndString = "}";
+				FString ResultString = String;
+				int32 StartIndex = ResultString.Find(StartString, ESearchCase::CaseSensitive);
+				while (StartIndex != INDEX_NONE)
+				{
+					StartIndex += StartString.Len();
+					int32 EndIndex = ResultString.Find(EndString, ESearchCase::CaseSensitive, ESearchDir::FromStart, StartIndex);
+					if (EndIndex > StartIndex)
+					{
+						FString EnvironmentVariableName = ResultString.Mid(StartIndex, EndIndex - StartIndex);
+						FString EnvironmentVariableValue = FPlatformMisc::GetEnvironmentVariable(*EnvironmentVariableName);
+						if (EnvironmentVariableValue.IsEmpty())
+						{
+							return false;
+						}
+						else
+						{
+							ResultString.ReplaceInline(*(StartString + EnvironmentVariableName + EndString), *EnvironmentVariableValue, ESearchCase::CaseSensitive);
+						}
+					}
+					else
+					{
+						return false;
+					}
+					StartIndex = ResultString.Find(StartString, ESearchCase::CaseSensitive);
+				}
+				String = ResultString;
+				return true;
+			}
+
+			FString GetSharedLibraryEntryPointName(const FString& HeaderString)
+			{
+				FString SearchString = "iree_hal_executable_library_header_t**";
+				int32 Start = HeaderString.Find(SearchString);
+				if (Start == INDEX_NONE)
+				{
+					return "";
+				}
+				Start += SearchString.Len();
+				int32 End = HeaderString.Find("(", ESearchCase::CaseSensitive, ESearchDir::FromStart, Start);
+				if (End <= Start)
+				{
+					return "";
+				}
+				return HeaderString.Mid(Start, End - Start).TrimStartAndEnd();
+			}
+
+			void RunCommand(const FString& Command, const FString& Arguments)
+			{
+				void* PipeRead = nullptr;
+				void* PipeWrite = nullptr;
+				FPlatformProcess::CreatePipe(PipeRead, PipeWrite);
+				FProcHandle ProcHandle = FPlatformProcess::CreateProc(*Command, *Arguments, false, true, true, nullptr, 0, nullptr, PipeWrite, PipeRead);
+				FPlatformProcess::WaitForProc(ProcHandle);
+				FPlatformProcess::CloseProc(ProcHandle);
+				FPlatformProcess::ClosePipe(PipeRead, PipeWrite);
+			}
+		} // Private
+
+		FCompiler::FCompiler(const FString& InCompilerCommand, const FString& InLinkerCommand, const FString& InSharedLibExt, TConstArrayView<FBuildTarget> InBuildTargets) : CompilerCommand(InCompilerCommand), LinkerCommand(InLinkerCommand), SharedLibExt(InSharedLibExt), BuildTargets(InBuildTargets)
+		{
+
+		}
+
+		TUniquePtr<FCompiler> FCompiler::Make(const FString& InTargetPlatformName)
+		{
+			using namespace Private;
+
+			FString PluginDir = FPaths::ConvertRelativePathToFull(*IPluginManager::Get().FindPlugin(UE_PLUGIN_NAME)->GetBaseDir());
+			FString BuildConfigFileName = FString("IREE_") + UGameplayStatics::GetPlatformName() + "_To_" + InTargetPlatformName + ".json";
+			TArray<FString> BuildConfigFilePaths =
+			{
+				FPaths::Combine(FPaths::ConvertRelativePathToFull(*FPaths::ProjectConfigDir()), BuildConfigFileName),
+				FPaths::Combine(PluginDir, "Config", BuildConfigFileName),
+				FPaths::Combine(FPaths::ConvertRelativePathToFull(*FPaths::EngineDir()), "Platforms", InTargetPlatformName, "Plugins", UE_PLUGIN_NAME, "Config", BuildConfigFileName),
+				FPaths::Combine(FPaths::ConvertRelativePathToFull(*FPaths::EngineDir()), "Platforms", InTargetPlatformName, "Plugins", "Experimental", UE_PLUGIN_NAME, "Config", BuildConfigFileName)
+			};
+
+			FString CompilerCommand;
+			FString LinkerCommand;
+			FString SharedLibExt;
+			TArray<FBuildTarget> BuildTargets;
+			IPlatformFile& PlatformFile = FPlatformFileManager::Get().GetPlatformFile();
+			for (FString BuildConfigFilePath : BuildConfigFilePaths)
+			{
+				if (PlatformFile.FileExists(*BuildConfigFilePath))
+				{
+					FString BuildConfigFileString;
+					if (FFileHelper::LoadFileToString(BuildConfigFileString, *BuildConfigFilePath))
+					{
+						TSharedRef<TJsonReader<TCHAR>> JsonReader = TJsonReaderFactory<TCHAR>::Create(BuildConfigFileString);
+						TSharedPtr<FJsonObject> JsonObject = MakeShareable(new FJsonObject);
+						FBuildConfig BuildConfig;
+						if (FJsonSerializer::Deserialize(JsonReader, JsonObject) && JsonObject.IsValid() && BuildConfig.FromJson(JsonObject))
+						{
+							if (BuildConfig.BuildTargets.IsEmpty())
+							{
+								UE_LOG(LogTemp, Warning, TEXT("UNNERuntimeIREECpu could not find targets in %s"), *BuildConfigFilePath);
+								continue;
+							}
+
+							FString TmpCompilerCommand;
+							for (int32 i = 0; i < BuildConfig.CompilerCommand.Num(); i++)
+							{
+								if (ResolveEnvironmentVariables(BuildConfig.CompilerCommand[i]))
+								{
+									BuildConfig.CompilerCommand[i].ReplaceInline(*FString("${PLUGIN_DIR}"), *PluginDir);
+									BuildConfig.CompilerCommand[i].ReplaceInline(*FString("${PROJECT_DIR}"), *FPaths::ProjectDir());
+									if (PlatformFile.FileExists(*BuildConfig.CompilerCommand[i]))
+									{
+										TmpCompilerCommand = BuildConfig.CompilerCommand[i];
+										break;
+									}
+								}
+								else
+								{
+									UE_LOG(LogTemp, Warning, TEXT("UNNERuntimeIREECpu could not replace environment variables in %s"), *BuildConfig.CompilerCommand[i]);
+								}
+							}
+							if (TmpCompilerCommand.IsEmpty())
+							{
+								UE_LOG(LogTemp, Warning, TEXT("UNNERuntimeIREECpu could not find the compiler executable in %s"), *BuildConfigFilePath);
+								continue;
+							}
+
+							FString TmpLinkerCommand;
+							for (int32 i = 0; i < BuildConfig.LinkerCommand.Num(); i++)
+							{
+								if (ResolveEnvironmentVariables(BuildConfig.LinkerCommand[i]))
+								{
+									BuildConfig.LinkerCommand[i].ReplaceInline(*FString("${PLUGIN_DIR}"), *PluginDir);
+									BuildConfig.LinkerCommand[i].ReplaceInline(*FString("${PROJECT_DIR}"), *FPaths::ProjectDir());
+									if (PlatformFile.FileExists(*BuildConfig.LinkerCommand[i]))
+									{
+										TmpLinkerCommand = BuildConfig.LinkerCommand[i];
+										break;
+									}
+								}
+								else
+								{
+									UE_LOG(LogTemp, Warning, TEXT("UNNERuntimeIREECpu could not replace environment variables in %s"), *BuildConfig.LinkerCommand[i]);
+								}
+							}
+							if (TmpLinkerCommand.IsEmpty())
+							{
+								UE_LOG(LogTemp, Warning, TEXT("UNNERuntimeIREECpu could not find the linker executable in %s"), *BuildConfigFilePath);
+								continue;
+							}
+
+							CompilerCommand = TmpCompilerCommand;
+							LinkerCommand = TmpLinkerCommand;
+							SharedLibExt = BuildConfig.SharedLibExt;
+							BuildTargets = BuildConfig.BuildTargets;
+							break;
+						}
+						else
+						{
+							UE_LOG(LogTemp, Warning, TEXT("UNNERuntimeIREECpu could not parse build config file %s"), *BuildConfigFilePath);
+						}
+					}
+					else
+					{
+						UE_LOG(LogTemp, Warning, TEXT("UNNERuntimeIREECpu could not read build config file %s"), *BuildConfigFilePath);
+					}
+				}
+			}
+			if (CompilerCommand.IsEmpty() || LinkerCommand.IsEmpty() || BuildTargets.IsEmpty())
+			{
+				return TUniquePtr<FCompiler>();
+			}
+			return TUniquePtr<FCompiler>(new FCompiler(CompilerCommand, LinkerCommand, SharedLibExt, BuildTargets));
+		}
+
+		bool FCompiler::CompileMlir(TConstArrayView<uint8> InFileData, const FString& InModelName, const FString& InIntermediateDir, const FString& InStagingDir, TArray<FCompilerResult>& OutCompilerResults, UNNERuntimeIREEModuleMetaData* ModuleMetaData)
+		{
+			using namespace Private;
+
+			FString InputFilePath = FPaths::Combine(InIntermediateDir, InModelName) + ".mlir";
+			FFileHelper::SaveArrayToFile(InFileData, *InputFilePath);
+
+			if (ModuleMetaData)
+			{
+				FString FileDataString = "";
+				FileDataString.AppendChars((char*)InFileData.GetData(), InFileData.Num());
+				ModuleMetaData->ParseFromString(FileDataString);
+			}
+
+			TArray<FCompilerResult> Results;
+			bool bResult = true;
+			IPlatformFile& PlatformFile = FPlatformFileManager::Get().GetPlatformFile();
+			for (int32 i = 0; i < BuildTargets.Num(); i++)
+			{
+				FString IntermediateDirPath = FPaths::Combine(InIntermediateDir, BuildTargets[i].Architecture);
+				PlatformFile.CreateDirectoryTree(*IntermediateDirPath);
+				FString IntermediateFilePathNoExt = FPaths::Combine(IntermediateDirPath, InModelName);
+				FString ObjectFilePath = IntermediateFilePathNoExt + ".o";
+				FString VmfbFilePath = IntermediateFilePathNoExt + ".vmfb";
+				FString SharedLibFilePath = IntermediateFilePathNoExt + SharedLibExt;
+
+				FString CompilerArguments = BuildTargets[i].CompilerArguments;
+				if (!ResolveEnvironmentVariables(CompilerArguments))
+				{
+					UE_LOG(LogTemp, Warning, TEXT("UNNERuntimeIREECpu could not replace environment variables in %s"), *BuildTargets[i].CompilerArguments);
+					continue;
+				}
+				CompilerArguments.ReplaceInline(*FString("${OBJECT_PATH}"), *(FString("\"") + ObjectFilePath + "\""));
+				CompilerArguments.ReplaceInline(*FString("${VMFB_PATH}"), *(FString("\"") + VmfbFilePath + "\""));
+				CompilerArguments.ReplaceInline(*FString("${INPUT_PATH}"), *(FString("\"") + InputFilePath + "\""));
+
+				RunCommand(CompilerCommand, CompilerArguments);
+
+				if (!PlatformFile.FileExists(*ObjectFilePath) || !PlatformFile.FileExists(*VmfbFilePath))
+				{
+					UE_LOG(LogTemp, Warning, TEXT("UNNERuntimeIREECpu failed to compile the model \"%s\" using the command:"), *InputFilePath);
+					UE_LOG(LogTemp, Warning, TEXT("\"%s\" %s"), *CompilerCommand, *CompilerArguments);
+					bResult = false;
+					continue;
+				}
+
+				FString LinkerArguments = BuildTargets[i].LinkerArguments;
+				if (!ResolveEnvironmentVariables(LinkerArguments))
+				{
+					UE_LOG(LogTemp, Warning, TEXT("UNNERuntimeIREECpu could not replace environment variables in %s"), *BuildTargets[i].LinkerArguments);
+					bResult = false;
+					continue;
+				}
+				LinkerArguments.ReplaceInline(*FString("${OBJECT_PATH}"), *(FString("\"") + ObjectFilePath + "\""));
+				LinkerArguments.ReplaceInline(*FString("${SHARED_LIB_PATH}"), *(FString("\"") + SharedLibFilePath + "\""));
+
+				RunCommand(LinkerCommand, LinkerArguments);
+
+				if (!PlatformFile.FileExists(*SharedLibFilePath))
+				{
+					UE_LOG(LogTemp, Warning, TEXT("UNNERuntimeIREECpu failed to link the model \"%s\" using the command:"), *InputFilePath);
+					UE_LOG(LogTemp, Warning, TEXT("\"%s\" %s"), *LinkerCommand, *LinkerArguments);
+					bResult = false;
+					continue;
+				}
+
+				FString StagingFilePathNoExt = FPaths::Combine(InStagingDir, BuildTargets[i].Architecture, InModelName);
+				FString StagedVmfbFilePath = StagingFilePathNoExt + ".vmfb";
+				FString StagedSharedLibFilePath = StagingFilePathNoExt + SharedLibExt;
+				if (IFileManager::Get().Copy(*StagedSharedLibFilePath, *SharedLibFilePath) != COPY_OK)
+				{
+					UE_LOG(LogTemp, Warning, TEXT("UNNERuntimeIREECpu failed to copy \"%s\" to \"%s\""), *SharedLibFilePath, *StagedSharedLibFilePath);
+					bResult = false;
+				}
+				if (IFileManager::Get().Copy(*StagedVmfbFilePath, *VmfbFilePath) != COPY_OK)
+				{
+					UE_LOG(LogTemp, Warning, TEXT("UNNERuntimeIREECpu failed to copy \"%s\" to \"%s\""), *VmfbFilePath, *StagedVmfbFilePath);
+					bResult = false;
+				}
+
+				FString SharedLibraryEntryPointName = "";
+				FString HeaderPath = IntermediateFilePathNoExt + ".h";
+				if (!PlatformFile.FileExists(*HeaderPath))
+				{
+					UE_LOG(LogTemp, Warning, TEXT("UNNERuntimeIREECpu could not find the model header \"%s\""), *HeaderPath);
+					bResult = false;
+					continue;
+				}
+				FString HeaderString;
+				if (!FFileHelper::LoadFileToString(HeaderString, *HeaderPath) || HeaderString.IsEmpty())
+				{
+					UE_LOG(LogTemp, Warning, TEXT("UNNERuntimeIREECpu could not read the model header \"%s\""), *HeaderPath);
+					bResult = false;
+					continue;
+				}
+				SharedLibraryEntryPointName = GetSharedLibraryEntryPointName(HeaderString);
+				if (SharedLibraryEntryPointName.IsEmpty())
+				{
+					UE_LOG(LogTemp, Warning, TEXT("UNNERuntimeIREECpu could not find the entry point in model header \"%s\""), *HeaderPath);
+					bResult = false;
+					continue;
+				}
+
+				FCompilerResult Result;
+				Result.Architecture = BuildTargets[i].Architecture;
+				Result.RelativeDirPath = BuildTargets[i].Architecture;
+				Result.SharedLibraryFileName = InModelName + SharedLibExt;
+				Result.VmfbFileName = InModelName + ".vmfb";
+				Result.SharedLibraryEntryPointName = SharedLibraryEntryPointName;
+				Results.Add(Result);
+			}
+
+			bResult &= !Results.IsEmpty();
+			if (bResult)
+			{
+				OutCompilerResults = Results;
+			}
+			return bResult;
+		}
+	} // CPU
+} // UE::NNERuntimeIREE
+
+#endif // WITH_EDITOR
+#endif // WITH_NNE_RUNTIME_IREE
\ No newline at end of file
diff --git a/Engine/Plugins/Experimental/NNERuntimeIREE/Source/NNERuntimeIREE/Private/NNERuntimeIREECompiler.h b/Engine/Plugins/Experimental/NNERuntimeIREE/Source/NNERuntimeIREE/Private/NNERuntimeIREECompiler.h
new file mode 100644
--- /dev/null
+++ b/Engine/Plugins/Experimental/NNERuntimeIREE/Source/NNERuntimeIREE/Private/NNERuntimeIREECompiler.h
@@ -0,0 +1,72 @@
+// Copyright Epic Games, Inc. All Rights Reserved.
+
+#pragma once
+
+#ifdef WITH_NNE_RUNTIME_IREE
+#if WITH_EDITOR
+
+#include "CoreMinimal.h"
+#include "NNERuntimeIREEMetaData.h"
+#include "Serialization/JsonSerializerMacros.h"
+
+namespace UE::NNERuntimeIREE::CPU
+{
+	struct FBuildTarget : public FJsonSerializable
+	{
+		FString Architecture;
+		FString CompilerArguments;
+		FString LinkerArguments;
+
+		BEGIN_JSON_SERIALIZER
+			JSON_SERIALIZE("Architecture", Architecture);
+			JSON_SERIALIZE("CompilerArguments", CompilerArguments);
+			JSON_SERIALIZE("LinkerArguments", LinkerArguments);
+		END_JSON_SERIALIZER
+	};
+
+	struct FBuildConfig : public FJsonSerializable
+	{
+		TArray<FString> CompilerCommand;
+		TArray<FString> LinkerCommand;
+		FString SharedLibExt;
+		TArray<FBuildTarget> BuildTargets;
+
+		BEGIN_JSON_SERIALIZER
+			JSON_SERIALIZE_ARRAY("CompilerCommand", CompilerCommand);
+			JSON_SERIALIZE_ARRAY("LinkerCommand", LinkerCommand);
+			JSON_SERIALIZE("SharedLibExt", SharedLibExt);
+			JSON_SERIALIZE_ARRAY_SERIALIZABLE("Targets", BuildTargets, FBuildTarget);
+		END_JSON_SERIALIZER
+	};
+
+	struct FCompilerResult
+	{
+		FString Architecture;
+		FString RelativeDirPath;
+		FString SharedLibraryFileName;
+		FString VmfbFileName;
+		FString SharedLibraryEntryPointName;
+	};
+
+	class FCompiler
+	{
+	private:
+		FCompiler(const FString& InCompilerCommand, const FString& InLinkerCommand, const FString& InSharedLibExt, TConstArrayView<FBuildTarget> InBuildTargets);
+
+	public:
+		~FCompiler() {};
+		static TUniquePtr<FCompiler> Make(const FString& InTargetPlatformName);
+
+	public:
+		bool CompileMlir(TConstArrayView<uint8> InFileData, const FString& InModelName, const FString& InIntermediateDir, const FString& InStagingDir, TArray<FCompilerResult>& OutCompilerResults, UNNERuntimeIREEModuleMetaData* ModuleMetaData);
+
+	private:
+		FString CompilerCommand;
+		FString LinkerCommand;
+		FString SharedLibExt;
+		TArray<FBuildTarget> BuildTargets;
+	};
+} // UE::NNERuntimeIREE::CPU
+
+#endif // WITH_EDITOR
+#endif // WITH_NNE_RUNTIME_IREE
\ No newline at end of file
diff --git a/Engine/Plugins/Experimental/NNERuntimeIREE/Source/NNERuntimeIREE/Private/NNERuntimeIREECpu.cpp b/Engine/Plugins/Experimental/NNERuntimeIREE/Source/NNERuntimeIREE/Private/NNERuntimeIREECpu.cpp
deleted file mode 100644
diff --git a/Engine/Plugins/Experimental/NNERuntimeIREE/Source/NNERuntimeIREE/Private/NNERuntimeIREECpu.h b/Engine/Plugins/Experimental/NNERuntimeIREE/Source/NNERuntimeIREE/Private/NNERuntimeIREECpu.h
deleted file mode 100644
diff --git a/Engine/Plugins/Experimental/NNERuntimeIREE/Source/NNERuntimeIREE/Private/NNERuntimeIREECpuCompiler.cpp b/Engine/Plugins/Experimental/NNERuntimeIREE/Source/NNERuntimeIREE/Private/NNERuntimeIREECpuCompiler.cpp
deleted file mode 100644
diff --git a/Engine/Plugins/Experimental/NNERuntimeIREE/Source/NNERuntimeIREE/Private/NNERuntimeIREECpuCompiler.h b/Engine/Plugins/Experimental/NNERuntimeIREE/Source/NNERuntimeIREE/Private/NNERuntimeIREECpuCompiler.h
deleted file mode 100644
diff --git a/Engine/Plugins/Experimental/NNERuntimeIREE/Source/NNERuntimeIREE/Private/NNERuntimeIREECpuModel.cpp b/Engine/Plugins/Experimental/NNERuntimeIREE/Source/NNERuntimeIREE/Private/NNERuntimeIREECpuModel.cpp
deleted file mode 100644
diff --git a/Engine/Plugins/Experimental/NNERuntimeIREE/Source/NNERuntimeIREE/Private/NNERuntimeIREECpuModel.h b/Engine/Plugins/Experimental/NNERuntimeIREE/Source/NNERuntimeIREE/Private/NNERuntimeIREECpuModel.h
deleted file mode 100644
diff --git a/Engine/Plugins/Experimental/NNERuntimeIREE/Source/NNERuntimeIREE/Private/NNERuntimeIREEGpu.cpp b/Engine/Plugins/Experimental/NNERuntimeIREE/Source/NNERuntimeIREE/Private/NNERuntimeIREEGpu.cpp
deleted file mode 100644
diff --git a/Engine/Plugins/Experimental/NNERuntimeIREE/Source/NNERuntimeIREE/Private/NNERuntimeIREEGpu.h b/Engine/Plugins/Experimental/NNERuntimeIREE/Source/NNERuntimeIREE/Private/NNERuntimeIREEGpu.h
deleted file mode 100644
diff --git a/Engine/Plugins/Experimental/NNERuntimeIREE/Source/NNERuntimeIREE/Private/NNERuntimeIREEMetaData.cpp b/Engine/Plugins/Experimental/NNERuntimeIREE/Source/NNERuntimeIREE/Private/NNERuntimeIREEMetaData.cpp
--- a/Engine/Plugins/Experimental/NNERuntimeIREE/Source/NNERuntimeIREE/Private/NNERuntimeIREEMetaData.cpp
+++ b/Engine/Plugins/Experimental/NNERuntimeIREE/Source/NNERuntimeIREE/Private/NNERuntimeIREEMetaData.cpp
@@ -2,6 +2,7 @@
 
 #include "NNERuntimeIREEMetaData.h"
 
+#include "Containers/Array.h"
 #include "Serialization/CustomVersion.h"
 
 namespace UE::NNERuntimeIREE::ModuleMetaData::Private
@@ -237,7 +238,7 @@ namespace UE::NNERuntimeIREE::ModuleMetaData::Private
 
 		return true;
 	}
-}
+} // UE::NNERuntimeIREE::ModuleMetaData::Private
 
 void UNNERuntimeIREEModuleMetaData::Serialize(FArchive& Ar)
 {
diff --git a/Engine/Plugins/Experimental/NNERuntimeIREE/Source/NNERuntimeIREE/Private/NNERuntimeIREEModel.cpp b/Engine/Plugins/Experimental/NNERuntimeIREE/Source/NNERuntimeIREE/Private/NNERuntimeIREEModel.cpp
new file mode 100644
--- /dev/null
+++ b/Engine/Plugins/Experimental/NNERuntimeIREE/Source/NNERuntimeIREE/Private/NNERuntimeIREEModel.cpp
@@ -0,0 +1,940 @@
+// Copyright Epic Games, Inc. All Rights Reserved.
+
+#include "NNERuntimeIREEModel.h"
+
+#ifdef WITH_NNE_RUNTIME_IREE
+
+#include "GenericPlatform/GenericPlatformProcess.h"
+#include "HAL/FileManager.h"
+#include "Memory/SharedBuffer.h"
+#include "NNEModelData.h"
+#include "Misc/Paths.h"
+#include "Misc/ScopeLock.h"
+#include "Serialization/Archive.h"
+
+#if PLATFORM_MICROSOFT
+#include "Microsoft/AllowMicrosoftPlatformTypes.h"
+#include "Microsoft/AllowMicrosoftPlatformAtomics.h"
+#endif // PLATFORM_MICROSOFT
+THIRD_PARTY_INCLUDES_START
+#include "iree/hal/drivers/local_sync/sync_device.h"
+#include "iree/hal/local/loaders/static_library_loader.h"
+#include "iree/runtime/call.h"
+#include "iree/runtime/instance.h"
+#include "iree/runtime/session.h"
+#include "iree/vm/bytecode/module.h"
+THIRD_PARTY_INCLUDES_END
+#if PLATFORM_MICROSOFT
+#include "Microsoft/HideMicrosoftPlatformAtomics.h"
+#include "Microsoft/HideMicrosoftPlatformTypes.h"
+#endif // PLATFORM_MICROSOF
+
+namespace UE::NNERuntimeIREE
+{
+	namespace Private
+	{
+		iree_hal_element_types_t NNEToIREEType(ENNETensorDataType Type)
+		{
+			switch (Type)
+			{
+			case ENNETensorDataType::None:
+				return IREE_HAL_ELEMENT_TYPE_NONE;
+				break;
+			case ENNETensorDataType::Char:
+				return IREE_HAL_ELEMENT_TYPE_UINT_8;
+				break;
+			case ENNETensorDataType::Boolean:
+				return IREE_HAL_ELEMENT_TYPE_BOOL_8;
+				break;
+			case ENNETensorDataType::Half:
+				return IREE_HAL_ELEMENT_TYPE_FLOAT_16;
+				break;
+			case ENNETensorDataType::Float:
+				return IREE_HAL_ELEMENT_TYPE_FLOAT_32;
+				break;
+			case ENNETensorDataType::Double:
+				return IREE_HAL_ELEMENT_TYPE_FLOAT_64;
+				break;
+			case ENNETensorDataType::Int8:
+				return IREE_HAL_ELEMENT_TYPE_INT_8;
+				break;
+			case ENNETensorDataType::Int16:
+				return IREE_HAL_ELEMENT_TYPE_INT_16;
+				break;
+			case ENNETensorDataType::Int32:
+				return IREE_HAL_ELEMENT_TYPE_INT_32;
+				break;
+			case ENNETensorDataType::Int64:
+				return IREE_HAL_ELEMENT_TYPE_INT_64;
+				break;
+			case ENNETensorDataType::UInt8:
+				return IREE_HAL_ELEMENT_TYPE_UINT_8;
+				break;
+			case ENNETensorDataType::UInt16:
+				return IREE_HAL_ELEMENT_TYPE_UINT_16;
+				break;
+			case ENNETensorDataType::UInt32:
+				return IREE_HAL_ELEMENT_TYPE_UINT_32;
+				break;
+			case ENNETensorDataType::UInt64:
+				return IREE_HAL_ELEMENT_TYPE_UINT_64;
+				break;
+			case ENNETensorDataType::Complex64:
+				return IREE_HAL_ELEMENT_TYPE_COMPLEX_FLOAT_64;
+				break;
+			case ENNETensorDataType::Complex128:
+				return IREE_HAL_ELEMENT_TYPE_COMPLEX_FLOAT_128;
+				break;
+			case ENNETensorDataType::BFloat16:
+				return IREE_HAL_ELEMENT_TYPE_BFLOAT_16;
+				break;
+			default:
+				return IREE_HAL_ELEMENT_TYPE_NONE;
+				break;
+			}
+		}
+
+		void PrintIREEError(const FString& InMessage, iree_status_t InStatus)
+		{
+			iree_host_size_t TrueLength = 0;
+			iree_status_format(InStatus, 0, (char*)nullptr, &TrueLength);
+			void* ErrorString = FMemory::Malloc(TrueLength + 1);
+			((char*)ErrorString)[TrueLength] = (char)0;
+			iree_status_format(InStatus, TrueLength, (char*)ErrorString, &TrueLength);
+			UE_LOG(LogTemp, Error, TEXT("%s: %s"), *InMessage, *FString(StringCast<TCHAR>(static_cast<const ANSICHAR*>(ErrorString)).Get()));
+			FMemory::Free(ErrorString);
+		}
+
+		class FInstance
+		{
+		private:
+			iree_runtime_instance_t* Instance;
+			static TWeakPtr<FInstance> WeakInstancePtr;
+			static FCriticalSection CriticalSection;
+
+			FInstance(iree_runtime_instance_t* InInstance) : Instance(InInstance)
+			{
+				check(InInstance);
+			}
+
+		public:
+			~FInstance()
+			{
+				iree_runtime_instance_release(Instance);
+			}
+
+			static TSharedPtr<FInstance> GetInstance()
+			{
+				FScopeLock ScopeLock(&CriticalSection);
+
+				if (WeakInstancePtr.IsValid())
+				{
+					return WeakInstancePtr.Pin();
+				}
+
+				iree_status_t Status = iree_ok_status();
+				if (!iree_status_is_ok(Status))
+				{
+					iree_status_free(Status);
+					return TSharedPtr<FInstance>();
+				}
+
+				iree_runtime_instance_options_t InstanceOptions;
+				iree_runtime_instance_options_initialize(&InstanceOptions);
+				iree_runtime_instance_options_use_all_available_drivers(&InstanceOptions);
+
+				iree_runtime_instance_t* Instance = nullptr;
+				Status = iree_runtime_instance_create(&InstanceOptions, iree_allocator_system(), &Instance);
+				if (!iree_status_is_ok(Status))
+				{
+					Private::PrintIREEError("UE::NNERuntimeIREE::Private::FInstance failed to create the instance", Status);
+
+					if (Instance)
+					{
+						iree_runtime_instance_release(Instance);
+					}
+
+					iree_status_free(Status);
+					return TSharedPtr<FInstance>();
+				}
+
+				TSharedPtr<FInstance> SharedInstance = TSharedPtr<FInstance>(new FInstance(Instance));
+				WeakInstancePtr = SharedInstance;
+
+				iree_status_free(Status);
+				return SharedInstance;
+			}
+
+			bool CreateModule(TConstArrayView<uint8> InVmfbDataView, iree_vm_module_t** OutModule)
+			{
+				FScopeLock ScopeLock(&CriticalSection);
+
+				check(!InVmfbDataView.IsEmpty());
+				check(OutModule);
+
+				iree_status_t Status = iree_ok_status();
+				check(iree_status_is_ok(Status));
+
+				const iree_const_byte_span_t ModuleData = iree_make_const_byte_span(InVmfbDataView.GetData(), InVmfbDataView.Num());
+
+				iree_vm_module_t* Module = nullptr;
+				Status = iree_vm_bytecode_module_create(iree_runtime_instance_vm_instance(Instance), ModuleData, iree_allocator_null(), GetHostAllocator(), &Module);
+				if (!iree_status_is_ok(Status))
+				{
+					Private::PrintIREEError("UE::NNERuntimeIREE::Private::FInstance failed to create the module", Status);
+
+					if (Module)
+					{
+						iree_vm_module_release(Module);
+					}
+
+					iree_status_free(Status);
+					return false;
+				}
+
+				*OutModule = Module;
+
+				iree_status_free(Status);
+				return true;
+			}
+
+			bool CreateSyncDevice(void* InLibraryQueryFuntionPointer, iree_hal_device_t** OutDevice)
+			{
+				FScopeLock ScopeLock(&CriticalSection);
+
+				check(InLibraryQueryFuntionPointer);
+				check(OutDevice);
+
+				iree_status_t Status = iree_ok_status();
+				check(iree_status_is_ok(Status));
+
+				iree_allocator_t HostAllocator = GetHostAllocator();
+
+				iree_hal_executable_loader_t* LibraryLoader = nullptr;
+				const iree_hal_executable_library_query_fn_t LibraryList[] = { (iree_hal_executable_library_query_fn_t)InLibraryQueryFuntionPointer };
+				Status = iree_hal_static_library_loader_create(IREE_ARRAYSIZE(LibraryList), LibraryList, iree_hal_executable_import_provider_null(), HostAllocator, &LibraryLoader);
+				if (!iree_status_is_ok(Status))
+				{
+					Private::PrintIREEError("UE::NNERuntimeIREE::Private::FInstance failed to create the library loader", Status);
+
+					if (LibraryLoader)
+					{
+						iree_hal_executable_loader_release(LibraryLoader);
+					}
+
+					iree_status_free(Status);
+					return false;
+				}
+
+				iree_hal_allocator_t* DeviceAllocator = nullptr;
+				iree_string_view_t Identifier = iree_make_cstring_view("sync");
+				Status = iree_hal_allocator_create_heap(Identifier, HostAllocator, HostAllocator, &DeviceAllocator);
+				if (!iree_status_is_ok(Status))
+				{
+					Private::PrintIREEError("UE::NNERuntimeIREE::Private::FInstance failed to create the device allocator", Status);
+
+					if (DeviceAllocator)
+					{
+						iree_hal_allocator_release(DeviceAllocator);
+					}
+
+					iree_hal_executable_loader_release(LibraryLoader);
+					iree_status_free(Status);
+					return false;
+				}
+
+				iree_hal_device_t* Device = nullptr;
+				iree_hal_sync_device_params_t DeviceParams;
+				iree_hal_sync_device_params_initialize(&DeviceParams);
+				Status = iree_hal_sync_device_create(Identifier, &DeviceParams, 1, &LibraryLoader, DeviceAllocator, HostAllocator, &Device);
+				if (!iree_status_is_ok(Status))
+				{
+					Private::PrintIREEError("UE::NNERuntimeIREE::Private::FInstance failed to create the device", Status);
+
+					if (Device)
+					{
+						iree_hal_device_release(Device);
+					}
+
+					iree_hal_allocator_release(DeviceAllocator);
+					iree_hal_executable_loader_release(LibraryLoader);
+					iree_status_free(Status);
+					return false;
+				}
+
+				iree_hal_allocator_release(DeviceAllocator);
+				iree_hal_executable_loader_release(LibraryLoader);
+
+				*OutDevice = Device;
+
+				iree_status_free(Status);
+				return true;
+			}
+
+			bool CreateSession(iree_hal_device_t* InDevice, iree_runtime_session_t** OutSession)
+			{
+				FScopeLock ScopeLock(&CriticalSection);
+
+				check(InDevice);
+				check(OutSession);
+
+				iree_status_t Status = iree_ok_status();
+				check(iree_status_is_ok(Status));
+
+				iree_runtime_session_options_t SessionOptions;
+				iree_runtime_session_options_initialize(&SessionOptions);
+
+				iree_runtime_session_t* Session = nullptr;
+				Status = iree_runtime_session_create_with_device(Instance, &SessionOptions, InDevice, GetHostAllocator(), &Session);
+				if (!iree_status_is_ok(Status))
+				{
+					Private::PrintIREEError("UE::NNERuntimeIREE::Private::FInstance failed to create the session", Status);
+
+					if (Session)
+					{
+						iree_runtime_session_release(Session);
+					}
+
+					iree_status_free(Status);
+					return false;
+				}
+
+				*OutSession = Session;
+
+				iree_status_free(Status);
+				return true;
+			}
+
+			iree_allocator_t GetHostAllocator()
+			{
+				FScopeLock ScopeLock(&CriticalSection);
+				return iree_runtime_instance_host_allocator(Instance);
+			}
+		};
+		TWeakPtr<FInstance> FInstance::WeakInstancePtr;
+		FCriticalSection FInstance::CriticalSection;
+
+		class FModule
+		{
+		private:
+			TSharedRef<FInstance> Instance;
+			TSharedRef<UE::NNE::FSharedModelData> ModelData;
+			iree_vm_module_t* Module;
+			TArray<UE::NNERuntimeIREE::FFunctionMetaData> FunctionMetaData;
+			static TMap<FString, TWeakPtr<FModule>> Modules;
+
+			FModule(TSharedRef<FInstance> InInstance, TSharedRef<UE::NNE::FSharedModelData> InModelData, iree_vm_module_t* InModule, TConstArrayView<UE::NNERuntimeIREE::FFunctionMetaData> InFunctionMetaData)
+				: Instance(InInstance), ModelData(InModelData), Module(InModule), FunctionMetaData(InFunctionMetaData)
+			{
+				check(!ModelData->GetView().IsEmpty());
+				check(InModule);
+				check(!InFunctionMetaData.IsEmpty());
+			}
+
+		public:
+			~FModule()
+			{
+				iree_vm_module_release(Module);
+			}
+
+			static TSharedPtr<FModule> Make(const FString& InVmfbPath, const FString& InVmfbName, const UNNERuntimeIREEModuleMetaData& InModuleMetaData)
+			{
+				check(!InVmfbName.IsEmpty());
+				check(!InModuleMetaData.FunctionMetaData.IsEmpty());
+
+				FString CombinedPath = FPaths::Combine(InVmfbPath, InVmfbName);
+				if (Modules.Contains(CombinedPath) && Modules[CombinedPath].IsValid())
+				{
+					return Modules[CombinedPath].Pin();
+				}
+
+				TUniquePtr<FArchive> Reader = TUniquePtr<FArchive>(IFileManager::Get().CreateFileReader(*CombinedPath, 0));
+				if (!Reader)
+				{
+					UE_LOG(LogTemp, Error, TEXT("UE::NNERuntimeIREE::Private::FModule failed to open the vmfb data file '%s'"), *CombinedPath);
+					return TSharedPtr<FModule>();
+				}
+				int64 DataSize = Reader->TotalSize();
+				if (DataSize < 1)
+				{
+					UE_LOG(LogTemp, Error, TEXT("UE::NNERuntimeIREE::Private::FModule's vmfb data file '%s' is empty"), *CombinedPath);
+					return TSharedPtr<FModule>();
+				}
+
+				void* Data = FMemory::Malloc(DataSize, IREE_HAL_HEAP_BUFFER_ALIGNMENT);
+				Reader->Serialize(Data, DataSize);
+				TSharedPtr<UE::NNE::FSharedModelData> ModelData = MakeShared<UE::NNE::FSharedModelData>(FSharedBuffer::TakeOwnership(Data, DataSize, FMemory::Free), IREE_HAL_HEAP_BUFFER_ALIGNMENT);
+				if (!ModelData.IsValid())
+				{
+					return TSharedPtr<FModule>();
+				}
+
+				TSharedPtr<FInstance> Instance = FInstance::GetInstance();
+				if (!Instance.IsValid())
+				{
+					return TSharedPtr<FModule>();
+				}
+
+				iree_vm_module_t* Module = nullptr;
+				if (!Instance->CreateModule(ModelData->GetView(), &Module))
+				{
+					return TSharedPtr<FModule>();
+				}
+
+				TSharedPtr<FModule> Result = TSharedPtr<FModule>(new FModule(Instance.ToSharedRef(), ModelData.ToSharedRef(), Module, InModuleMetaData.FunctionMetaData));
+				Modules.Add(CombinedPath, Result);
+				return Result;
+			}
+
+			bool AppendToSession(iree_runtime_session_t* InSession)
+			{
+				check(InSession);
+
+				iree_status_t Status = iree_ok_status();
+				check(iree_status_is_ok(Status));
+
+				Status = iree_runtime_session_append_module(InSession, Module);
+				if (!iree_status_is_ok(Status))
+				{
+					PrintIREEError("UE::NNERuntimeIREE::Private::FModule failed to append the module to the session", Status);
+					iree_status_free(Status);
+					return false;
+				}
+
+				iree_status_free(Status);
+				return true;
+			}
+
+			TConstArrayView<UE::NNERuntimeIREE::FFunctionMetaData> GetFunctionMetaDataView()
+			{
+				return FunctionMetaData;
+			}
+
+			bool GetFunctionByName(const FString& InFunctionName, iree_vm_function_t* OutFunction)
+			{
+				iree_status_t Status = iree_ok_status();
+				check(iree_status_is_ok(Status));
+
+				bool bFound = false;
+				iree_host_size_t Ordinal = 0;
+				iree_vm_function_t Function;
+				do 
+				{
+					Status = iree_vm_module_lookup_function_by_ordinal(Module, IREE_VM_FUNCTION_LINKAGE_EXPORT, Ordinal, &Function);
+					if (iree_status_is_ok(Status))
+					{
+						Ordinal++;
+						iree_string_view_t FunctionNameView = iree_vm_function_name(&Function);
+						FString FunctionName = FString::ConstructFromPtrSize(FunctionNameView.data, FunctionNameView.size);
+						if (FunctionName.Equals(InFunctionName))
+						{
+							bFound = true;
+							break;
+						}
+					}
+				} while (iree_status_is_ok(Status));
+
+				if (!bFound)
+				{
+					UE_LOG(LogTemp, Error, TEXT("UE::NNERuntimeIREE::Private::FModule failed to find the module function %s"), *InFunctionName);
+					iree_status_free(Status);
+					return false;
+				}
+
+				*OutFunction = Function;
+				iree_status_free(Status);
+				return true;
+			}
+		};
+		TMap<FString, TWeakPtr<FModule>> FModule::Modules;
+	} // Private
+
+	namespace CPU
+	{
+		namespace Private
+		{
+			class FLibrary
+			{
+			private:
+				void* Library;
+				static TMap<FString, TWeakPtr<FLibrary>> Libraries;
+
+				FLibrary(void* InLibrary) : Library(InLibrary)
+				{
+					check(InLibrary);
+				}
+
+			public:
+				~FLibrary()
+				{
+					FPlatformProcess::FreeDllHandle(Library);
+				}
+
+				static TSharedPtr<FLibrary> GetLibrary(const FString& InLibraryPath, const FString& InLibraryName)
+				{
+					check(!InLibraryName.IsEmpty());
+
+					FString CombinedPath = FPaths::Combine(InLibraryPath, InLibraryName);
+					if (Libraries.Contains(CombinedPath) && Libraries[CombinedPath].IsValid())
+					{
+						return Libraries[CombinedPath].Pin();
+					}
+
+#ifdef NNE_RUNTIME_IREE_USE_COMBINED_LIB_PATH
+					void* Library = FPlatformProcess::GetDllHandle(*CombinedPath);
+					if (!Library)
+					{
+						UE_LOG(LogTemp, Error, TEXT("UE::NNERuntimeIREE::CPU::Private::FLibrary failed to load the shared library '%s'"), *CombinedPath);
+						return TSharedPtr<FLibrary>();
+					}
+#else
+					FPlatformProcess::PushDllDirectory(*InLibraryPath);
+					void* Library = FPlatformProcess::GetDllHandle(*InLibraryName);
+					FPlatformProcess::PopDllDirectory(*InLibraryPath);
+					if (!Library)
+					{
+						UE_LOG(LogTemp, Error, TEXT("UE::NNERuntimeIREE::CPU::Private::FLibrary failed to load the shared library '%s' from '%s'"), *InLibraryName, *InLibraryPath);
+						return TSharedPtr<FLibrary>();
+					}
+#endif
+
+					TSharedPtr<FLibrary> Result = TSharedPtr<FLibrary>(new FLibrary(Library));
+					Libraries.Add(CombinedPath, Result);
+					return Result;
+				}
+
+				bool GetFunctionPointer(const FString& InFunctionName, void** OutFunctionPointer) const
+				{
+					void* Result = FPlatformProcess::GetDllExport(Library, *InFunctionName);
+					if (Result)
+					{
+						*OutFunctionPointer = Result;
+						return true;
+					}
+					UE_LOG(LogTemp, Error, TEXT("UE::NNERuntimeIREE::CPU::Private::FLibrary failed to get the function %s"), *InFunctionName);
+					return false;
+				}
+			};
+			TMap<FString, TWeakPtr<FLibrary>> FLibrary::Libraries;
+
+			class FDevice
+			{
+			private:
+				TSharedRef<UE::NNERuntimeIREE::Private::FInstance> Instance;
+				TSharedRef<FLibrary> Library;
+				iree_hal_device_t* Device;
+				static TMap<FString, TWeakPtr<FDevice>> Devices;
+
+				FDevice(TSharedRef<UE::NNERuntimeIREE::Private::FInstance> InInstance, TSharedRef<FLibrary> InLibrary, iree_hal_device_t* InDevice) : Instance(InInstance), Library(InLibrary), Device(InDevice)
+				{
+					check(InDevice);
+				}
+
+			public:
+				~FDevice()
+				{
+					iree_hal_device_release(Device);
+				}
+
+				static TSharedPtr<FDevice> Make(const FString& InLibraryPath, const FString& InLibraryName, const FString& LibraryQueryFunctionName)
+				{
+					check(!InLibraryName.IsEmpty());
+					check(!LibraryQueryFunctionName.IsEmpty());
+
+					FString CombinedPathPlusFunction = FPaths::Combine(InLibraryPath, InLibraryName) + "::" + LibraryQueryFunctionName;
+					if (Devices.Contains(CombinedPathPlusFunction) && Devices[CombinedPathPlusFunction].IsValid())
+					{
+						return Devices[CombinedPathPlusFunction].Pin();
+					}
+
+					TSharedPtr<FLibrary> Library = FLibrary::GetLibrary(InLibraryPath, InLibraryName);
+					if (!Library.IsValid())
+					{
+						return TSharedPtr<FDevice>();
+					}
+
+					void* LibraryQueryFunctionPointer = nullptr;
+					Library->GetFunctionPointer(LibraryQueryFunctionName, &LibraryQueryFunctionPointer);
+					if (!LibraryQueryFunctionPointer)
+					{
+						return TSharedPtr<FDevice>();
+					}
+
+					TSharedPtr<UE::NNERuntimeIREE::Private::FInstance> Instance = UE::NNERuntimeIREE::Private::FInstance::GetInstance();
+					if (!Instance.IsValid())
+					{
+						return TSharedPtr<FDevice>();
+					}
+
+					iree_hal_device_t* Device = nullptr;
+					if (!Instance->CreateSyncDevice(LibraryQueryFunctionPointer, &Device))
+					{
+						return TSharedPtr<FDevice>();
+					}
+
+					TSharedPtr<FDevice> Result = TSharedPtr<FDevice>(new FDevice(Instance.ToSharedRef(), Library.ToSharedRef(), Device));
+					Devices.Add(CombinedPathPlusFunction, Result);
+					return Result;
+				}
+
+				bool CreateSession(iree_runtime_session_t** OutSession)
+				{
+					return Instance->CreateSession(Device, OutSession);
+				}
+
+				iree_hal_allocator_t* GetDeviceAllocator()
+				{
+					return iree_hal_device_allocator(Device);
+				}
+			};
+			TMap<FString, TWeakPtr<FDevice>> FDevice::Devices;
+
+			class FSession
+			{
+			private:
+				TSharedRef<FDevice> Device;
+				TSharedRef<UE::NNERuntimeIREE::Private::FModule> Module;
+				iree_runtime_session_t* Session;
+				iree_runtime_call_t Call;
+				TArray<UE::NNE::FTensorDesc> InputTensorDescs;
+				TArray<UE::NNE::FTensorDesc> OutputTensorDescs;
+				TArray<UE::NNE::FTensorShape> InputTensorShapes;
+				TArray<UE::NNE::FTensorShape> OutputTensorShapes;
+
+				FSession(TSharedRef<FDevice> InDevice, TSharedRef<UE::NNERuntimeIREE::Private::FModule> InModule, iree_runtime_session_t* InSession, const iree_runtime_call_t& InCall, TConstArrayView<UE::NNE::FTensorDesc> InInputTensorDescs, TConstArrayView<UE::NNE::FTensorDesc> InOutputTensorDescs)
+					: Device(InDevice), Module(InModule), Session(InSession), Call(InCall), InputTensorDescs(InInputTensorDescs), OutputTensorDescs(InOutputTensorDescs)
+				{
+					check(InSession);
+
+					check(!InputTensorDescs.IsEmpty());
+					bool bAllConcrete = true;
+					for (int32 i = 0; i < InInputTensorDescs.Num(); i++)
+					{
+						bAllConcrete &= InInputTensorDescs[i].GetShape().IsConcrete();
+					}
+					if (bAllConcrete)
+					{
+						for (int32 i = 0; i < InInputTensorDescs.Num(); i++)
+						{
+							InputTensorShapes.Add(UE::NNE::FTensorShape::MakeFromSymbolic(InInputTensorDescs[i].GetShape()));
+						}
+					}
+
+					check(!OutputTensorDescs.IsEmpty());
+					bAllConcrete = true;
+					for (int32 i = 0; i < InOutputTensorDescs.Num(); i++)
+					{
+						bAllConcrete &= InOutputTensorDescs[i].GetShape().IsConcrete();
+					}
+					if (bAllConcrete)
+					{
+						for (int32 i = 0; i < InOutputTensorDescs.Num(); i++)
+						{
+							OutputTensorShapes.Add(UE::NNE::FTensorShape::MakeFromSymbolic(InOutputTensorDescs[i].GetShape()));
+						}
+					}
+				}
+
+			public:
+				~FSession()
+				{
+					iree_runtime_call_deinitialize(&Call);
+					iree_runtime_session_release(Session);
+				}
+
+				static TSharedPtr<FSession> Make(TSharedRef<FDevice> InDevice, TSharedRef<UE::NNERuntimeIREE::Private::FModule> InModule)
+				{
+					check(!InModule->GetFunctionMetaDataView().IsEmpty());
+
+					iree_runtime_session_t* Session = nullptr;
+					if (!InDevice->CreateSession(&Session))
+					{
+						return TSharedPtr<FSession>();
+					}
+
+					if (!InModule->AppendToSession(Session))
+					{
+						iree_runtime_session_release(Session);
+						return TSharedPtr<FSession>();
+					}
+
+					iree_status_t Status = iree_ok_status();
+					check(iree_status_is_ok(Status));
+
+					FString MainFunctionName = InModule->GetFunctionMetaDataView()[0].Name;
+					iree_vm_function_t MainFunction;
+					if (!InModule->GetFunctionByName(MainFunctionName, &MainFunction))
+					{
+						iree_runtime_session_release(Session);
+						iree_status_free(Status);
+						return TSharedPtr<FSession>();
+					}
+
+					iree_host_size_t NumInputs = 0;
+					iree_host_size_t NumOutputs = 0;
+					iree_vm_function_signature_t Signature = iree_vm_function_signature(&MainFunction);
+					Status = iree_vm_function_call_count_arguments_and_results(&Signature, &NumInputs, &NumOutputs);
+					TConstArrayView<UE::NNE::FTensorDesc> InputTensorDescs = InModule->GetFunctionMetaDataView()[0].InputDescs;
+					TConstArrayView<UE::NNE::FTensorDesc> OutputTensorDescs = InModule->GetFunctionMetaDataView()[0].OutputDescs;
+					if (!iree_status_is_ok(Status) || NumInputs != InputTensorDescs.Num() || NumOutputs != OutputTensorDescs.Num())
+					{
+						UE_LOG(LogTemp, Error, TEXT("UE::NNERuntimeIREE::CPU::Private::FSession has a function signature mismatch in function %s"), *MainFunctionName);
+						iree_runtime_session_release(Session);
+						iree_status_free(Status);
+						return TSharedPtr<FSession>();
+					}
+
+					iree_runtime_call_t Call;
+					Status = iree_runtime_call_initialize(Session, MainFunction, &Call);
+					if (!iree_status_is_ok(Status))
+					{
+						UE::NNERuntimeIREE::Private::PrintIREEError("UE::NNERuntimeIREE::CPU::Private::FSession failed to initialize the session call", Status);
+						iree_runtime_session_release(Session);
+						iree_status_free(Status);
+						return TSharedPtr<FSession>();
+					}
+
+					TSharedPtr<FSession> Result = TSharedPtr<FSession>(new FSession(InDevice, InModule, Session, Call, InputTensorDescs, OutputTensorDescs));
+					iree_status_free(Status);
+					return Result;
+				}
+
+				TConstArrayView<UE::NNE::FTensorDesc> GetInputTensorDescs() const
+				{
+					return InputTensorDescs;
+				}
+
+				TConstArrayView<UE::NNE::FTensorDesc> GetOutputTensorDescs() const
+				{
+					return OutputTensorDescs;
+				}
+
+				TConstArrayView<UE::NNE::FTensorShape> GetInputTensorShapes() const
+				{
+					return InputTensorShapes;
+				}
+
+				TConstArrayView<UE::NNE::FTensorShape> GetOutputTensorShapes() const
+				{
+					return OutputTensorShapes;
+				}
+
+				int32 SetInputTensorShapes(TConstArrayView<UE::NNE::FTensorShape> InInputShapes)
+				{
+					check(InputTensorDescs.Num() == InInputShapes.Num());
+					checkCode(for (int32 i = 0; i < InputTensorDescs.Num(); i++) { check(InInputShapes[i].IsCompatibleWith(InputTensorDescs[i].GetShape())); });
+					InputTensorShapes = InInputShapes;
+					return 0;
+				}
+
+				int32 RunSyncCPU(TConstArrayView<UE::NNE::FTensorBindingCPU> InInputBindings, TConstArrayView<UE::NNE::FTensorBindingCPU> InOutputBindings)
+				{
+					check(InInputBindings.Num() == InputTensorShapes.Num());
+
+					iree_status_t Status = iree_ok_status();
+					check(iree_status_is_ok(Status));
+
+					iree_runtime_call_reset(&Call);
+
+					for (int32 i = 0; i < InInputBindings.Num(); i++)
+					{
+						check(InInputBindings[i].SizeInBytes >= InputTensorShapes[i].Volume() * InputTensorDescs[i].GetElementByteSize());
+						checkf(FMath::Modulo<uint64>((uint64)InInputBindings[i].Data, IREE_HAL_HEAP_BUFFER_ALIGNMENT) == 0, TEXT("NNERuntimeIREECpu requires input- and output-buffer memory to be aligned with %d bytes"), IREE_HAL_HEAP_BUFFER_ALIGNMENT);
+
+						iree_hal_buffer_view_t* TempBufferView;
+						iree_hal_buffer_params_t Params = { 0 };
+						Params.type = IREE_HAL_MEMORY_TYPE_DEVICE_LOCAL;
+						Params.usage = IREE_HAL_BUFFER_USAGE_DEFAULT;
+						iree_hal_dim_t Shape[UE::NNE::FTensorShape::MaxRank];
+						for (int32 j = 0; j < InputTensorShapes[i].Rank(); j++)
+						{
+							Shape[j] = InputTensorShapes[i].GetData()[j];
+						}
+						ENNETensorDataType NNEType = InputTensorDescs[i].GetDataType();
+						iree_hal_element_types_t IREEType = UE::NNERuntimeIREE::Private::NNEToIREEType(NNEType);
+						Status = iree_hal_buffer_view_allocate_buffer(
+							Device->GetDeviceAllocator(),
+							InputTensorShapes[i].Rank(), Shape,
+							IREEType, IREE_HAL_ENCODING_TYPE_DENSE_ROW_MAJOR,
+							Params,
+							iree_make_const_byte_span((void*)InInputBindings[i].Data, InInputBindings[i].SizeInBytes),
+							&TempBufferView);
+						if (!iree_status_is_ok(Status))
+						{
+							UE::NNERuntimeIREE::Private::PrintIREEError("UE::NNERuntimeIREE::CPU::Private::FSession failed to allocate the buffer view", Status);
+							if (TempBufferView)
+							{
+								iree_hal_buffer_view_release(TempBufferView);
+							}
+							iree_status_free(Status);
+							return -1;
+						}
+
+						Status = iree_runtime_call_inputs_push_back_buffer_view(&Call, TempBufferView);
+						iree_hal_buffer_view_release(TempBufferView);
+						if (!iree_status_is_ok(Status))
+						{
+							UE::NNERuntimeIREE::Private::PrintIREEError("UE::NNERuntimeIREE::CPU::Private::FSession failed to push the buffer view to the input list", Status);
+							iree_status_free(Status);
+							return -1;
+						}
+					}
+
+					Status = iree_runtime_call_invoke(&Call, 0);
+					if (!iree_status_is_ok(Status))
+					{
+						UE::NNERuntimeIREE::Private::PrintIREEError("UE::NNERuntimeIREE::CPU::Private::FSession failed to call the model function", Status);
+						iree_status_free(Status);
+						return -1;
+					}
+
+					OutputTensorShapes.Reset();
+					TArray<iree_hal_buffer_view_t*> BufferViews;
+					iree_hal_buffer_view_t* BufferView = nullptr;
+					Status = iree_runtime_call_outputs_pop_front_buffer_view(&Call, &BufferView);
+					while (iree_status_is_ok(Status))
+					{
+						iree_host_size_t Rank = iree_hal_buffer_view_shape_rank(BufferView);
+						const iree_hal_dim_t* Dims = iree_hal_buffer_view_shape_dims(BufferView);
+						uint32 Shape[UE::NNE::FTensorShape::MaxRank];
+						for (int32 i = 0; i < FMath::Min((int32)Rank, UE::NNE::FTensorShape::MaxRank); i++)
+						{
+							Shape[i] = (uint32)Dims[i];
+						}
+						OutputTensorShapes.Add(UE::NNE::FTensorShape::Make(TConstArrayView<uint32>(Shape, FMath::Min((int32)Rank, UE::NNE::FTensorShape::MaxRank))));
+
+						BufferViews.Add(BufferView);
+						Status = iree_runtime_call_outputs_pop_front_buffer_view(&Call, &BufferView);
+					}
+
+					bool bCopyResults = true;
+					if (InOutputBindings.Num() != OutputTensorShapes.Num())
+					{
+						bCopyResults = false;
+					}
+					for (int32 i = 0; i < InOutputBindings.Num() && bCopyResults; i++)
+					{
+						if (InOutputBindings[i].SizeInBytes < iree_hal_buffer_view_element_size(BufferViews[i]) * iree_hal_buffer_view_element_count(BufferViews[i]))
+						{
+							bCopyResults = false;
+						}
+					}
+					int32 Result = 0;
+					if (bCopyResults)
+					{
+						for (int32 i = 0; i < InOutputBindings.Num(); i++)
+						{
+							iree_hal_buffer_t* Buffer = iree_hal_buffer_view_buffer(BufferViews[i]);
+							if (!Buffer)
+							{
+								UE_LOG(LogTemp, Error, TEXT("UE::NNERuntimeIREE::CPU::Private::FSession failed to get the result buffer"));
+								Result = -1;
+								break;
+							}
+
+							int32 DataSizeInBytes = iree_hal_buffer_view_element_size(BufferViews[i]) * iree_hal_buffer_view_element_count(BufferViews[i]);
+
+							iree_hal_buffer_mapping_t BufferMapping;
+							Status = iree_hal_buffer_map_range(Buffer, IREE_HAL_MAPPING_MODE_PERSISTENT, IREE_HAL_MEMORY_ACCESS_READ, 0, DataSizeInBytes, &BufferMapping);
+							if (!iree_status_is_ok(Status))
+							{
+								UE::NNERuntimeIREE::Private::PrintIREEError("UE::NNERuntimeIREE::CPU::Private::FSession failed to map the result buffer", Status);
+								Result = -1;
+								break;
+							}
+							FMemory::Memcpy(InOutputBindings[i].Data, BufferMapping.contents.data, DataSizeInBytes);
+							iree_hal_buffer_unmap_range(&BufferMapping);
+						}
+					}
+
+					for (int32 i = 0; i < InOutputBindings.Num(); i++)
+					{
+						iree_hal_buffer_view_release(BufferViews[i]);
+					}
+					iree_status_free(Status);
+					return Result;
+				}
+			};
+		}
+
+		FModelInstance::FModelInstance(TSharedRef<UE::NNERuntimeIREE::CPU::Private::FSession> InSession) : Session(InSession)
+		{
+			
+		}
+
+		TSharedPtr<FModelInstance> FModelInstance::Make(TSharedRef<UE::NNERuntimeIREE::CPU::Private::FDevice> InDevice, TSharedRef<UE::NNERuntimeIREE::Private::FModule> InModule)
+		{
+			TSharedPtr<Private::FSession> Session = Private::FSession::Make(InDevice, InModule);
+			if (!Session.IsValid())
+			{
+				return TSharedPtr<FModelInstance>();
+			}
+
+			return TSharedPtr<FModelInstance>(new FModelInstance(Session.ToSharedRef()));
+		}
+
+		TConstArrayView<UE::NNE::FTensorDesc> FModelInstance::GetInputTensorDescs() const
+		{
+			return Session->GetInputTensorDescs();
+		}
+
+		TConstArrayView<UE::NNE::FTensorDesc> FModelInstance::GetOutputTensorDescs() const
+		{
+			return Session->GetOutputTensorDescs();
+		}
+
+		TConstArrayView<UE::NNE::FTensorShape> FModelInstance::GetInputTensorShapes() const
+		{
+			return Session->GetInputTensorShapes();
+		}
+
+		TConstArrayView<UE::NNE::FTensorShape> FModelInstance::GetOutputTensorShapes() const
+		{
+			return Session->GetOutputTensorShapes();
+		}
+
+		int32 FModelInstance::SetInputTensorShapes(TConstArrayView<UE::NNE::FTensorShape> InInputShapes)
+		{
+			return Session->SetInputTensorShapes(InInputShapes);
+		}
+
+		int32 FModelInstance::RunSync(TConstArrayView<UE::NNE::FTensorBindingCPU> InInputBindings, TConstArrayView<UE::NNE::FTensorBindingCPU> InOutputBindings)
+		{
+			return Session->RunSyncCPU(InInputBindings, InOutputBindings);
+		}
+
+		FModel::FModel(TSharedRef<UE::NNERuntimeIREE::CPU::Private::FDevice> InDevice, TSharedRef<UE::NNERuntimeIREE::Private::FModule> InModule) : Device(InDevice), Module(InModule)
+		{
+
+		}
+
+		TSharedPtr<FModel> FModel::Make(const FString& InDirPath, const FString& InSharedLibraryFileName, const FString& InVmfbFileName, const FString& InLibraryQueryFunctionName, const UNNERuntimeIREEModuleMetaData& InModuleMetaData)
+		{
+			check(!InSharedLibraryFileName.IsEmpty());
+			check(!InVmfbFileName.IsEmpty());
+			check(!InLibraryQueryFunctionName.IsEmpty());
+			check(!InModuleMetaData.FunctionMetaData.IsEmpty());
+
+			TSharedPtr<Private::FDevice> Device = Private::FDevice::Make(InDirPath, InSharedLibraryFileName, InLibraryQueryFunctionName);
+			if (!Device.IsValid())
+			{
+				return TSharedPtr<FModel>();
+			}
+
+			TSharedPtr<UE::NNERuntimeIREE::Private::FModule> Module = UE::NNERuntimeIREE::Private::FModule::Make(InDirPath, InVmfbFileName, InModuleMetaData);
+			if (!Module.IsValid())
+			{
+				return TSharedPtr<FModel>();
+			}
+
+			return TSharedPtr<FModel>(new FModel(Device.ToSharedRef(), Module.ToSharedRef()));
+		}
+
+		TSharedPtr<UE::NNE::IModelInstanceCPU> FModel::CreateModelInstanceCPU()
+		{
+			return FModelInstance::Make(Device, Module);
+		}
+	} // CPU
+} // UE::NNERuntimeIREE
+
+#endif // WITH_NNE_RUNTIME_IREE
\ No newline at end of file
diff --git a/Engine/Plugins/Experimental/NNERuntimeIREE/Source/NNERuntimeIREE/Private/NNERuntimeIREEModel.h b/Engine/Plugins/Experimental/NNERuntimeIREE/Source/NNERuntimeIREE/Private/NNERuntimeIREEModel.h
new file mode 100644
--- /dev/null
+++ b/Engine/Plugins/Experimental/NNERuntimeIREE/Source/NNERuntimeIREE/Private/NNERuntimeIREEModel.h
@@ -0,0 +1,67 @@
+// Copyright Epic Games, Inc. All Rights Reserved.
+
+#pragma once
+
+#ifdef WITH_NNE_RUNTIME_IREE
+
+#include "NNERuntimeCPU.h"
+#include "NNERuntimeIREEMetaData.h"
+
+namespace UE::NNERuntimeIREE
+{
+	namespace Private
+	{
+		class FModule;
+	} // Private
+
+	namespace CPU
+	{
+		namespace Private 
+		{
+			class FSession;
+			class FDevice;
+		} // Private
+
+		class FModelInstance : public UE::NNE::IModelInstanceCPU
+		{
+		private:
+			FModelInstance(TSharedRef<UE::NNERuntimeIREE::CPU::Private::FSession> InSession);
+
+		public:
+			static TSharedPtr<FModelInstance> Make(TSharedRef<UE::NNERuntimeIREE::CPU::Private::FDevice> InDevice, TSharedRef<UE::NNERuntimeIREE::Private::FModule> InModule);
+
+		public:
+			//~ Begin IModelInstanceCPU Interface
+			virtual TConstArrayView<UE::NNE::FTensorDesc> GetInputTensorDescs() const override;
+			virtual TConstArrayView<UE::NNE::FTensorDesc> GetOutputTensorDescs() const override;
+			virtual TConstArrayView<UE::NNE::FTensorShape> GetInputTensorShapes() const override;
+			virtual TConstArrayView<UE::NNE::FTensorShape> GetOutputTensorShapes() const override;
+			virtual int32 SetInputTensorShapes(TConstArrayView<UE::NNE::FTensorShape> InInputShapes) override;
+			virtual int32 RunSync(TConstArrayView<UE::NNE::FTensorBindingCPU> InInputBindings, TConstArrayView<UE::NNE::FTensorBindingCPU> InOutputBindings) override;
+			//~ End IModelInstanceCPU Interface
+
+		private:
+			TSharedRef<UE::NNERuntimeIREE::CPU::Private::FSession> Session;
+		};
+
+		class FModel : public UE::NNE::IModelCPU
+		{
+		private:
+			FModel(TSharedRef<UE::NNERuntimeIREE::CPU::Private::FDevice> InDevice, TSharedRef<UE::NNERuntimeIREE::Private::FModule> InModule);
+
+		public:
+			static TSharedPtr<FModel> Make(const FString& InDirPath, const FString& InSharedLibraryFileName, const FString& InVmfbFileName, const FString& InLibraryQueryFunctionName, const UNNERuntimeIREEModuleMetaData& InModuleMetaData);
+
+		public:
+			//~ Begin IModelCPU Interface
+			virtual TSharedPtr<UE::NNE::IModelInstanceCPU> CreateModelInstanceCPU() override;
+			//~ End IModelCPU Interface
+
+		private:
+			TSharedRef<UE::NNERuntimeIREE::CPU::Private::FDevice> Device;
+			TSharedRef<UE::NNERuntimeIREE::Private::FModule> Module;
+		};
+	} // CPU
+} // UE::NNERuntimeIREE
+
+#endif // WITH_NNE_RUNTIME_IREE
\ No newline at end of file
diff --git a/Engine/Plugins/Experimental/NNERuntimeIREE/Source/NNERuntimeIREE/Private/NNERuntimeIREEModule.h b/Engine/Plugins/Experimental/NNERuntimeIREE/Source/NNERuntimeIREE/Private/NNERuntimeIREEModule.h
--- a/Engine/Plugins/Experimental/NNERuntimeIREE/Source/NNERuntimeIREE/Private/NNERuntimeIREEModule.h
+++ b/Engine/Plugins/Experimental/NNERuntimeIREE/Source/NNERuntimeIREE/Private/NNERuntimeIREEModule.h
@@ -5,9 +5,7 @@
 #include "Modules/ModuleInterface.h"
 
 #ifdef WITH_NNE_RUNTIME_IREE
-#include "NNERuntimeIREECpu.h"
-#include "NNERuntimeIREEGpu.h"
-#include "NNERuntimeIREERdg.h"
+#include "NNERuntimeIREE.h"
 #include "UObject/WeakObjectPtrTemplates.h"
 #endif // WITH_NNE_RUNTIME_IREE
 
diff --git a/Engine/Plugins/Experimental/NNERuntimeIREE/Source/NNERuntimeIREE/Private/NNERuntimeIREERdg.cpp b/Engine/Plugins/Experimental/NNERuntimeIREE/Source/NNERuntimeIREE/Private/NNERuntimeIREERdg.cpp
deleted file mode 100644
diff --git a/Engine/Plugins/Experimental/NNERuntimeIREE/Source/NNERuntimeIREE/Private/NNERuntimeIREERdg.h b/Engine/Plugins/Experimental/NNERuntimeIREE/Source/NNERuntimeIREE/Private/NNERuntimeIREERdg.h
deleted file mode 100644
