[NNE] Const Cleanup

Problem:
NNERuntimeORT is reaching Beta and we want a clean and lean API, however constness was not explicit in some of public API.

Solution:
Add constness indicator were missing.

Tests:
Win64 editor: style transfer + smoke test
Win64 standalone: style transfer + all tests

#rnx

diff --git a/Engine/Plugins/Experimental/NNERuntimeBasicCpu/Source/NNERuntimeBasicCpu/Private/NNERuntimeBasicCpu.cpp b/Engine/Plugins/Experimental/NNERuntimeBasicCpu/Source/NNERuntimeBasicCpu/Private/NNERuntimeBasicCpu.cpp
--- a/Engine/Plugins/Experimental/NNERuntimeBasicCpu/Source/NNERuntimeBasicCpu/Private/NNERuntimeBasicCpu.cpp
+++ b/Engine/Plugins/Experimental/NNERuntimeBasicCpu/Source/NNERuntimeBasicCpu/Private/NNERuntimeBasicCpu.cpp
@@ -61,7 +61,7 @@ FString UNNERuntimeBasicCpuImpl::GetModelDataIdentifier(const FString& FileType,
 	return FileId.ToString(EGuidFormats::Digits) + "-" + FString::FromInt(UE::NNE::RuntimeBasic::FModelCPU::ModelMagicNumber);
 }
 
-bool UNNERuntimeBasicCpuImpl::CanCreateModelCPU(TObjectPtr<UNNEModelData> ModelData) const
+bool UNNERuntimeBasicCpuImpl::CanCreateModelCPU(const TObjectPtr<UNNEModelData> ModelData) const
 {
 	check(ModelData != nullptr);
 
@@ -97,7 +97,7 @@ bool UNNERuntimeBasicCpuImpl::CanCreateModelCPU(TObjectPtr<UNNEModelData> ModelD
 	return true;
 }
 
-TSharedPtr<UE::NNE::IModelCPU> UNNERuntimeBasicCpuImpl::CreateModelCPU(TObjectPtr<UNNEModelData> ModelData)
+TSharedPtr<UE::NNE::IModelCPU> UNNERuntimeBasicCpuImpl::CreateModelCPU(const TObjectPtr<UNNEModelData> ModelData)
 {
 	check(ModelData != nullptr);
 
diff --git a/Engine/Plugins/Experimental/NNERuntimeBasicCpu/Source/NNERuntimeBasicCpu/Private/NNERuntimeBasicCpu.h b/Engine/Plugins/Experimental/NNERuntimeBasicCpu/Source/NNERuntimeBasicCpu/Private/NNERuntimeBasicCpu.h
--- a/Engine/Plugins/Experimental/NNERuntimeBasicCpu/Source/NNERuntimeBasicCpu/Private/NNERuntimeBasicCpu.h
+++ b/Engine/Plugins/Experimental/NNERuntimeBasicCpu/Source/NNERuntimeBasicCpu/Private/NNERuntimeBasicCpu.h
@@ -29,8 +29,8 @@ public:
 	virtual TSharedPtr<UE::NNE::FSharedModelData> CreateModelData(const FString& FileType, TConstArrayView<uint8> FileData, const TMap<FString, TConstArrayView<uint8>>& AdditionalFileData, const FGuid& FileId, const ITargetPlatform* TargetPlatform) override;
 	virtual FString GetModelDataIdentifier(const FString& FileType, TConstArrayView<uint8> FileData, const TMap<FString, TConstArrayView<uint8>>& AdditionalFileData, const FGuid& FileId, const ITargetPlatform* TargetPlatform) const override;
 
-	virtual bool CanCreateModelCPU(TObjectPtr<UNNEModelData> ModelData) const override;
-	virtual TSharedPtr<UE::NNE::IModelCPU> CreateModelCPU(TObjectPtr<UNNEModelData> ModelData) override;
+	virtual bool CanCreateModelCPU(const TObjectPtr<UNNEModelData> ModelData) const override;
+	virtual TSharedPtr<UE::NNE::IModelCPU> CreateModelCPU(const TObjectPtr<UNNEModelData> ModelData) override;
 
 private:
 
diff --git a/Engine/Plugins/Experimental/NNERuntimeIREE/Source/NNERuntimeIREE/Private/NNERuntimeIREE.cpp b/Engine/Plugins/Experimental/NNERuntimeIREE/Source/NNERuntimeIREE/Private/NNERuntimeIREE.cpp
--- a/Engine/Plugins/Experimental/NNERuntimeIREE/Source/NNERuntimeIREE/Private/NNERuntimeIREE.cpp
+++ b/Engine/Plugins/Experimental/NNERuntimeIREE/Source/NNERuntimeIREE/Private/NNERuntimeIREE.cpp
@@ -222,7 +222,7 @@ FString UNNERuntimeIREECpu::GetModelDataIdentifier(const FString& FileType, TCon
 	return UE::NNERuntimeIREE::CPU::Private::GetModelDataIdentifier(GetRuntimeName(), UNNERuntimeIREECpu::GUID, FileId.ToString(EGuidFormats::Digits), PlatformName, "");
 }
 
-bool UNNERuntimeIREECpu::CanCreateModelCPU(TObjectPtr<UNNEModelData> ModelData) const
+bool UNNERuntimeIREECpu::CanCreateModelCPU(const TObjectPtr<UNNEModelData> ModelData) const
 {
 	check(ModelData != nullptr);
 
@@ -245,7 +245,7 @@ bool UNNERuntimeIREECpu::CanCreateModelCPU(TObjectPtr<UNNEModelData> ModelData)
 	return bResult;
 }
 
-TSharedPtr<UE::NNE::IModelCPU> UNNERuntimeIREECpu::CreateModelCPU(TObjectPtr<UNNEModelData> ModelData)
+TSharedPtr<UE::NNE::IModelCPU> UNNERuntimeIREECpu::CreateModelCPU(const TObjectPtr<UNNEModelData> ModelData)
 {
 	check(ModelData != nullptr);
 
@@ -434,7 +434,7 @@ FString UNNERuntimeIREEGpu::GetModelDataIdentifier(const FString& FileType, TCon
 	return UE::NNERuntimeIREE::CPU::Private::GetModelDataIdentifier(GetRuntimeName(), GetGUID(), FileId.ToString(EGuidFormats::Digits), PlatformName, "");
 }
 
-bool UNNERuntimeIREEGpu::CanCreateModelGPU(TObjectPtr<UNNEModelData> ModelData) const
+bool UNNERuntimeIREEGpu::CanCreateModelGPU(const TObjectPtr<UNNEModelData> ModelData) const
 {
 	check(ModelData != nullptr);
 
@@ -459,7 +459,7 @@ bool UNNERuntimeIREEGpu::CanCreateModelGPU(TObjectPtr<UNNEModelData> ModelData)
 	return bResult;
 }
 
-TSharedPtr<UE::NNE::IModelGPU> UNNERuntimeIREEGpu::CreateModelGPU(TObjectPtr<UNNEModelData> ModelData)
+TSharedPtr<UE::NNE::IModelGPU> UNNERuntimeIREEGpu::CreateModelGPU(const TObjectPtr<UNNEModelData> ModelData)
 {
 	check(ModelData != nullptr);
 
@@ -575,7 +575,7 @@ FString UNNERuntimeIREERdg::GetModelDataIdentifier(const FString& FileType, TCon
 	return UE::NNERuntimeIREE::CPU::Private::GetModelDataIdentifier(GetRuntimeName(), UNNERuntimeIREERdg::GUID, FileId.ToString(EGuidFormats::Digits), PlatformName, "");
 }
 
-bool UNNERuntimeIREERdg::CanCreateModelRDG(TObjectPtr<UNNEModelData> ModelData) const
+bool UNNERuntimeIREERdg::CanCreateModelRDG(const TObjectPtr<UNNEModelData> ModelData) const
 {
 	check(ModelData != nullptr);
 
@@ -598,7 +598,7 @@ bool UNNERuntimeIREERdg::CanCreateModelRDG(TObjectPtr<UNNEModelData> ModelData)
 	return bResult;
 }
 
-TSharedPtr<UE::NNE::IModelRDG> UNNERuntimeIREERdg::CreateModelRDG(TObjectPtr<UNNEModelData> ModelData)
+TSharedPtr<UE::NNE::IModelRDG> UNNERuntimeIREERdg::CreateModelRDG(const TObjectPtr<UNNEModelData> ModelData)
 {
 	check(ModelData != nullptr);
 
diff --git a/Engine/Plugins/Experimental/NNERuntimeIREE/Source/NNERuntimeIREE/Private/NNERuntimeIREE.h b/Engine/Plugins/Experimental/NNERuntimeIREE/Source/NNERuntimeIREE/Private/NNERuntimeIREE.h
--- a/Engine/Plugins/Experimental/NNERuntimeIREE/Source/NNERuntimeIREE/Private/NNERuntimeIREE.h
+++ b/Engine/Plugins/Experimental/NNERuntimeIREE/Source/NNERuntimeIREE/Private/NNERuntimeIREE.h
@@ -30,8 +30,8 @@ public:
 	//~ End INNERuntime Interface
 
 	//~ Begin INNERuntimeCPU Interface
-	virtual bool CanCreateModelCPU(TObjectPtr<UNNEModelData> ModelData) const override;
-	virtual TSharedPtr<UE::NNE::IModelCPU> CreateModelCPU(TObjectPtr<UNNEModelData> ModelData) override;
+	virtual bool CanCreateModelCPU(const TObjectPtr<UNNEModelData> ModelData) const override;
+	virtual TSharedPtr<UE::NNE::IModelCPU> CreateModelCPU(const TObjectPtr<UNNEModelData> ModelData) override;
 	//~ End INNERuntimeCPU Interface
 
 	static void GetUpdatedPlatformConfig(const FString& PlatformName, FConfigFile& ConfigFile, FString& ConfigFilePath);
@@ -44,8 +44,8 @@ public:
 	//~ End INNERuntime Interface
 
 	//~ Begin INNERuntimeCPU Interface
-	virtual bool CanCreateModelCPU(TObjectPtr<UNNEModelData> ModelData) const override { return false; };
-	virtual TSharedPtr<UE::NNE::IModelCPU> CreateModelCPU(TObjectPtr<UNNEModelData> ModelData) override { return TSharedPtr<UE::NNE::IModelCPU>(); };
+	virtual bool CanCreateModelCPU(const TObjectPtr<UNNEModelData> ModelData) const override { return false; };
+	virtual TSharedPtr<UE::NNE::IModelCPU> CreateModelCPU(const TObjectPtr<UNNEModelData> ModelData) override { return TSharedPtr<UE::NNE::IModelCPU>(); };
 	//~ End INNERuntimeCPU Interface
 #endif // WITH_NNE_RUNTIME_IREE
 };
@@ -65,8 +65,8 @@ public:
 	//~ End INNERuntime Interface
 
 	//~ Begin INNERuntimeGPU Interface
-	virtual bool CanCreateModelGPU(TObjectPtr<UNNEModelData> ModelData) const override;
-	virtual TSharedPtr<UE::NNE::IModelGPU> CreateModelGPU(TObjectPtr<UNNEModelData> ModelData) override;
+	virtual bool CanCreateModelGPU(const TObjectPtr<UNNEModelData> ModelData) const override;
+	virtual TSharedPtr<UE::NNE::IModelGPU> CreateModelGPU(const TObjectPtr<UNNEModelData> ModelData) override;
 	//~ End INNERuntimeGPU Interface
 
 	virtual bool IsAvailable() const;
@@ -81,8 +81,8 @@ public:
 	//~ End INNERuntime Interface
 
 	//~ Begin INNERuntimeGPU Interface
-	virtual bool CanCreateModelGPU(TObjectPtr<UNNEModelData> ModelData) const override { return false; };
-	virtual TSharedPtr<UE::NNE::IModelGPU> CreateModelGPU(TObjectPtr<UNNEModelData> ModelData) override { return TSharedPtr<UE::NNE::IModelGPU>(); };
+	virtual bool CanCreateModelGPU(const TObjectPtr<UNNEModelData> ModelData) const override { return false; };
+	virtual TSharedPtr<UE::NNE::IModelGPU> CreateModelGPU(const TObjectPtr<UNNEModelData> ModelData) override { return TSharedPtr<UE::NNE::IModelGPU>(); };
 	//~ End INNERuntimeGPU Interface
 #endif // WITH_NNE_RUNTIME_IREE
 };
@@ -149,8 +149,8 @@ public:
 	//~ End INNERuntime Interface
 
 	//~ Begin INNERuntimeRdg Interface
-	virtual bool CanCreateModelRDG(TObjectPtr<UNNEModelData> ModelData) const override;
-	virtual TSharedPtr<UE::NNE::IModelRDG> CreateModelRDG(TObjectPtr<UNNEModelData> ModelData) override;
+	virtual bool CanCreateModelRDG(const TObjectPtr<UNNEModelData> ModelData) const override;
+	virtual TSharedPtr<UE::NNE::IModelRDG> CreateModelRDG(const TObjectPtr<UNNEModelData> ModelData) override;
 	//~ End INNERuntimeRdg Interface
 
 	bool IsAvailable() const;
@@ -163,8 +163,8 @@ public:
 	//~ End INNERuntime Interface
 
 	//~ Begin INNERuntimeRDG Interface
-	virtual bool CanCreateModelRDG(TObjectPtr<UNNEModelData> ModelData) const override { return false; };
-	virtual TSharedPtr<UE::NNE::IModelRDG> CreateModelRDG(TObjectPtr<UNNEModelData> ModelData) override { return TSharedPtr<UE::NNE::IModelRDG>(); };
+	virtual bool CanCreateModelRDG(const TObjectPtr<UNNEModelData> ModelData) const override { return false; };
+	virtual TSharedPtr<UE::NNE::IModelRDG> CreateModelRDG(const TObjectPtr<UNNEModelData> ModelData) override { return TSharedPtr<UE::NNE::IModelRDG>(); };
 	//~ End INNERuntimeRDG Interface
 #endif // WITH_NNE_RUNTIME_IREE
 };
\ No newline at end of file
diff --git a/Engine/Plugins/Experimental/NNERuntimeORTCpu/Source/NNERuntimeORTCpu/Private/NNERuntimeORTCpu.cpp b/Engine/Plugins/Experimental/NNERuntimeORTCpu/Source/NNERuntimeORTCpu/Private/NNERuntimeORTCpu.cpp
--- a/Engine/Plugins/Experimental/NNERuntimeORTCpu/Source/NNERuntimeORTCpu/Private/NNERuntimeORTCpu.cpp
+++ b/Engine/Plugins/Experimental/NNERuntimeORTCpu/Source/NNERuntimeORTCpu/Private/NNERuntimeORTCpu.cpp
@@ -56,7 +56,7 @@ FString UNNERuntimeORTCustomCpuImpl::GetModelDataIdentifier(const FString& FileT
 	return FileId.ToString(EGuidFormats::Digits) + "-" + UNNERuntimeORTCustomCpuImpl::GUID.ToString(EGuidFormats::Digits) + "-" + FString::FromInt(UNNERuntimeORTCustomCpuImpl::Version);
 }
 
-bool UNNERuntimeORTCustomCpuImpl::CanCreateModelCPU(TObjectPtr<UNNEModelData> ModelData) const
+bool UNNERuntimeORTCustomCpuImpl::CanCreateModelCPU(const TObjectPtr<UNNEModelData> ModelData) const
 {
 	check(ModelData != nullptr);
 	
@@ -81,7 +81,7 @@ bool UNNERuntimeORTCustomCpuImpl::CanCreateModelCPU(TObjectPtr<UNNEModelData> Mo
 	return bResult;
 }
 
-TSharedPtr<UE::NNE::IModelCPU> UNNERuntimeORTCustomCpuImpl::CreateModelCPU(TObjectPtr<UNNEModelData> ModelData)
+TSharedPtr<UE::NNE::IModelCPU> UNNERuntimeORTCustomCpuImpl::CreateModelCPU(const TObjectPtr<UNNEModelData> ModelData)
 {
 	check(ModelData != nullptr);
 	
diff --git a/Engine/Plugins/Experimental/NNERuntimeORTCpu/Source/NNERuntimeORTCpu/Private/NNERuntimeORTCpu.h b/Engine/Plugins/Experimental/NNERuntimeORTCpu/Source/NNERuntimeORTCpu/Private/NNERuntimeORTCpu.h
--- a/Engine/Plugins/Experimental/NNERuntimeORTCpu/Source/NNERuntimeORTCpu/Private/NNERuntimeORTCpu.h
+++ b/Engine/Plugins/Experimental/NNERuntimeORTCpu/Source/NNERuntimeORTCpu/Private/NNERuntimeORTCpu.h
@@ -37,6 +37,6 @@ public:
 	virtual TSharedPtr<UE::NNE::FSharedModelData> CreateModelData(const FString& FileType, TConstArrayView<uint8> FileData, const TMap<FString, TConstArrayView<uint8>>& AdditionalFileData, const FGuid& FileId, const ITargetPlatform* TargetPlatform) override;
 	virtual FString GetModelDataIdentifier(const FString& FileType, TConstArrayView<uint8> FileData, const TMap<FString, TConstArrayView<uint8>>& AdditionalFileData, const FGuid& FileId, const ITargetPlatform* TargetPlatform) const override;
 
-	virtual bool CanCreateModelCPU(TObjectPtr<UNNEModelData> ModelData) const override;
-	virtual TSharedPtr<UE::NNE::IModelCPU> CreateModelCPU(TObjectPtr<UNNEModelData> ModelData) override;
+	virtual bool CanCreateModelCPU(const TObjectPtr<UNNEModelData> ModelData) const override;
+	virtual TSharedPtr<UE::NNE::IModelCPU> CreateModelCPU(const TObjectPtr<UNNEModelData> ModelData) override;
 };
\ No newline at end of file
diff --git a/Engine/Plugins/Experimental/NNERuntimeRDG/Source/NNERuntimeRDG/Private/NNERuntimeRDGDml.cpp b/Engine/Plugins/Experimental/NNERuntimeRDG/Source/NNERuntimeRDG/Private/NNERuntimeRDGDml.cpp
--- a/Engine/Plugins/Experimental/NNERuntimeRDG/Source/NNERuntimeRDG/Private/NNERuntimeRDGDml.cpp
+++ b/Engine/Plugins/Experimental/NNERuntimeRDG/Source/NNERuntimeRDG/Private/NNERuntimeRDGDml.cpp
@@ -165,7 +165,7 @@ FString UNNERuntimeRDGDmlImpl::GetModelDataIdentifier(const FString& FileType, T
 	return FileId.ToString(EGuidFormats::Digits) + "-" + FModelInfo::Get()->GetGuid().ToString(EGuidFormats::Digits) + "-" + FString::FromInt(FModelInfo::Get()->GetVersion());
 }
 
-bool UNNERuntimeRDGDmlImpl::CanCreateModelRDG(TObjectPtr<UNNEModelData> ModelData) const
+bool UNNERuntimeRDGDmlImpl::CanCreateModelRDG(const TObjectPtr<UNNEModelData> ModelData) const
 {
 #ifdef NNE_USE_DIRECTML
 	if (bRegisterOnlyOperators)
@@ -194,7 +194,7 @@ bool UNNERuntimeRDGDmlImpl::CanCreateModelRDG(TObjectPtr<UNNEModelData> ModelDat
 #endif
 };
 
-TSharedPtr<UE::NNE::IModelRDG> UNNERuntimeRDGDmlImpl::CreateModelRDG(TObjectPtr<UNNEModelData> ModelData)
+TSharedPtr<UE::NNE::IModelRDG> UNNERuntimeRDGDmlImpl::CreateModelRDG(const TObjectPtr<UNNEModelData> ModelData)
 {
 #ifdef NNE_USE_DIRECTML
 	if (!CanCreateModelRDG(ModelData))
diff --git a/Engine/Plugins/Experimental/NNERuntimeRDG/Source/NNERuntimeRDG/Private/NNERuntimeRDGDml.h b/Engine/Plugins/Experimental/NNERuntimeRDG/Source/NNERuntimeRDG/Private/NNERuntimeRDGDml.h
--- a/Engine/Plugins/Experimental/NNERuntimeRDG/Source/NNERuntimeRDG/Private/NNERuntimeRDGDml.h
+++ b/Engine/Plugins/Experimental/NNERuntimeRDG/Source/NNERuntimeRDG/Private/NNERuntimeRDGDml.h
@@ -37,8 +37,8 @@ public:
 	virtual TSharedPtr<UE::NNE::FSharedModelData> CreateModelData(const FString& FileType, TConstArrayView<uint8> FileData, const TMap<FString, TConstArrayView<uint8>>& AdditionalFileData, const FGuid& FileId, const ITargetPlatform* TargetPlatform) override;
 	virtual FString GetModelDataIdentifier(const FString& FileType, TConstArrayView<uint8> FileData, const TMap<FString, TConstArrayView<uint8>>& AdditionalFileData, const FGuid& FileId, const ITargetPlatform* TargetPlatform) const override;
 
-	virtual bool CanCreateModelRDG(TObjectPtr<UNNEModelData> ModelData) const override;
-	virtual TSharedPtr<UE::NNE::IModelRDG> CreateModelRDG(TObjectPtr<UNNEModelData> ModelData) override;
+	virtual bool CanCreateModelRDG(const TObjectPtr<UNNEModelData> ModelData) const override;
+	virtual TSharedPtr<UE::NNE::IModelRDG> CreateModelRDG(const TObjectPtr<UNNEModelData> ModelData) override;
 
 private:
 	
diff --git a/Engine/Plugins/Experimental/NNERuntimeRDG/Source/NNERuntimeRDG/Private/NNERuntimeRDGHlsl.cpp b/Engine/Plugins/Experimental/NNERuntimeRDG/Source/NNERuntimeRDG/Private/NNERuntimeRDGHlsl.cpp
--- a/Engine/Plugins/Experimental/NNERuntimeRDG/Source/NNERuntimeRDG/Private/NNERuntimeRDGHlsl.cpp
+++ b/Engine/Plugins/Experimental/NNERuntimeRDG/Source/NNERuntimeRDG/Private/NNERuntimeRDGHlsl.cpp
@@ -95,7 +95,7 @@ bool UNNERuntimeRDGHlslImpl::CanCreateModelData(const FString& FileType, TConstA
 #endif
 }
 
-bool UNNERuntimeRDGHlslImpl::CanCreateModelRDG(TObjectPtr<UNNEModelData> ModelData) const
+bool UNNERuntimeRDGHlslImpl::CanCreateModelRDG(const TObjectPtr<UNNEModelData> ModelData) const
 {
 	int32 GuidSize = sizeof(GUID);
 	int32 VersionSize = sizeof(Version);
@@ -156,7 +156,7 @@ FString UNNERuntimeRDGHlslImpl::GetModelDataIdentifier(const FString& FileType,
 	return FileId.ToString(EGuidFormats::Digits) + "-" + UNNERuntimeRDGHlslImpl::GUID.ToString(EGuidFormats::Digits) + "-" + FString::FromInt(UNNERuntimeRDGHlslImpl::Version);
 }
 
-TSharedPtr<UE::NNE::IModelRDG> UNNERuntimeRDGHlslImpl::CreateModelRDG(TObjectPtr<UNNEModelData> ModelData)
+TSharedPtr<UE::NNE::IModelRDG> UNNERuntimeRDGHlslImpl::CreateModelRDG(const TObjectPtr<UNNEModelData> ModelData)
 {
 	if (!CanCreateModelRDG(ModelData))
 	{
diff --git a/Engine/Plugins/Experimental/NNERuntimeRDG/Source/NNERuntimeRDG/Private/NNERuntimeRDGHlsl.h b/Engine/Plugins/Experimental/NNERuntimeRDG/Source/NNERuntimeRDG/Private/NNERuntimeRDGHlsl.h
--- a/Engine/Plugins/Experimental/NNERuntimeRDG/Source/NNERuntimeRDG/Private/NNERuntimeRDGHlsl.h
+++ b/Engine/Plugins/Experimental/NNERuntimeRDG/Source/NNERuntimeRDG/Private/NNERuntimeRDGHlsl.h
@@ -32,6 +32,6 @@ public:
 	virtual TSharedPtr<UE::NNE::FSharedModelData> CreateModelData(const FString& FileType, TConstArrayView<uint8> FileData, const TMap<FString, TConstArrayView<uint8>>& AdditionalFileData, const FGuid& FileId, const ITargetPlatform* TargetPlatform) override;
 	virtual FString GetModelDataIdentifier(const FString& FileType, TConstArrayView<uint8> FileData, const TMap<FString, TConstArrayView<uint8>>& AdditionalFileData, const FGuid& FileId, const ITargetPlatform* TargetPlatform) const override;
 
-	virtual bool CanCreateModelRDG(TObjectPtr<UNNEModelData> ModelData) const override;
-	virtual TSharedPtr<UE::NNE::IModelRDG> CreateModelRDG(TObjectPtr<UNNEModelData> ModelData) override;
+	virtual bool CanCreateModelRDG(const TObjectPtr<UNNEModelData> ModelData) const override;
+	virtual TSharedPtr<UE::NNE::IModelRDG> CreateModelRDG(const TObjectPtr<UNNEModelData> ModelData) override;
 };
\ No newline at end of file
diff --git a/Engine/Plugins/NNE/NNERuntimeORT/Source/NNERuntimeORT/Private/NNERuntimeORT.cpp b/Engine/Plugins/NNE/NNERuntimeORT/Source/NNERuntimeORT/Private/NNERuntimeORT.cpp
--- a/Engine/Plugins/NNE/NNERuntimeORT/Source/NNERuntimeORT/Private/NNERuntimeORT.cpp
+++ b/Engine/Plugins/NNE/NNERuntimeORT/Source/NNERuntimeORT/Private/NNERuntimeORT.cpp
@@ -122,7 +122,7 @@ FString UNNERuntimeORTCpu::GetRuntimeName() const
 	return TEXT("NNERuntimeORTCpu");
 }
 
-bool UNNERuntimeORTCpu::CanCreateModelCPU(TObjectPtr<UNNEModelData> ModelData) const
+bool UNNERuntimeORTCpu::CanCreateModelCPU(const TObjectPtr<UNNEModelData> ModelData) const
 {
 	check(ModelData != nullptr);
 
@@ -147,7 +147,7 @@ bool UNNERuntimeORTCpu::CanCreateModelCPU(TObjectPtr<UNNEModelData> ModelData) c
 	return bResult;
 }
 
-TSharedPtr<UE::NNE::IModelCPU> UNNERuntimeORTCpu::CreateModelCPU(TObjectPtr<UNNEModelData> ModelData)
+TSharedPtr<UE::NNE::IModelCPU> UNNERuntimeORTCpu::CreateModelCPU(const TObjectPtr<UNNEModelData> ModelData)
 {
 	check(ModelData != nullptr);
 	check(ORTEnvironment.IsValid());
@@ -177,7 +177,7 @@ TSharedPtr<UE::NNE::IModelCPU> UNNERuntimeORTCpu::CreateModelCPU(TObjectPtr<UNNE
 }
 
 #if PLATFORM_WINDOWS
-bool UNNERuntimeORTDml::CanCreateModelGPU(TObjectPtr<UNNEModelData> ModelData) const
+bool UNNERuntimeORTDml::CanCreateModelGPU(const TObjectPtr<UNNEModelData> ModelData) const
 {
 	check(ModelData != nullptr);
 
@@ -211,7 +211,7 @@ bool UNNERuntimeORTDml::CanCreateModelGPU(TObjectPtr<UNNEModelData> ModelData) c
 	return bResult;
 }
 
-TSharedPtr<UE::NNE::IModelGPU> UNNERuntimeORTDml::CreateModelGPU(TObjectPtr<UNNEModelData> ModelData)
+TSharedPtr<UE::NNE::IModelGPU> UNNERuntimeORTDml::CreateModelGPU(const TObjectPtr<UNNEModelData> ModelData)
 {
 	check(ModelData != nullptr);
 	check(ORTEnvironment.IsValid());
@@ -242,12 +242,12 @@ TSharedPtr<UE::NNE::IModelGPU> UNNERuntimeORTDml::CreateModelGPU(TObjectPtr<UNNE
 
 #else // PLATFORM_WINDOWS
 
-bool UNNERuntimeORTDml::CanCreateModelGPU(TObjectPtr<UNNEModelData> ModelData) const
+bool UNNERuntimeORTDml::CanCreateModelGPU(const TObjectPtr<UNNEModelData> ModelData) const
 {
 	return false;
 }
 
-TSharedPtr<UE::NNE::IModelGPU> UNNERuntimeORTDml::CreateModelGPU(TObjectPtr<UNNEModelData> ModelData)
+TSharedPtr<UE::NNE::IModelGPU> UNNERuntimeORTDml::CreateModelGPU(const TObjectPtr<UNNEModelData> ModelData)
 {
 	return TSharedPtr<UE::NNE::IModelGPU>();
 }
diff --git a/Engine/Plugins/NNE/NNERuntimeORT/Source/NNERuntimeORT/Private/NNERuntimeORT.h b/Engine/Plugins/NNE/NNERuntimeORT/Source/NNERuntimeORT/Private/NNERuntimeORT.h
--- a/Engine/Plugins/NNE/NNERuntimeORT/Source/NNERuntimeORT/Private/NNERuntimeORT.h
+++ b/Engine/Plugins/NNE/NNERuntimeORT/Source/NNERuntimeORT/Private/NNERuntimeORT.h
@@ -38,8 +38,8 @@ public:
 	virtual TSharedPtr<UE::NNE::FSharedModelData> CreateModelData(const FString& FileType, TConstArrayView<uint8> FileData, const TMap<FString, TConstArrayView<uint8>>& AdditionalFileData, const FGuid& FileId, const ITargetPlatform* TargetPlatform) override;
 	virtual FString GetModelDataIdentifier(const FString& FileType, TConstArrayView<uint8> FileData, const TMap<FString, TConstArrayView<uint8>>& AdditionalFileData, const FGuid& FileId, const ITargetPlatform* TargetPlatform) const override;
 
-	virtual bool CanCreateModelGPU(TObjectPtr<UNNEModelData> ModelData) const override;
-	virtual TSharedPtr<UE::NNE::IModelGPU> CreateModelGPU(TObjectPtr<UNNEModelData> ModelData) override;
+	virtual bool CanCreateModelGPU(const TObjectPtr<UNNEModelData> ModelData) const override;
+	virtual TSharedPtr<UE::NNE::IModelGPU> CreateModelGPU(const TObjectPtr<UNNEModelData> ModelData) override;
 };
 
 UCLASS()
@@ -66,6 +66,6 @@ public:
 	virtual TSharedPtr<UE::NNE::FSharedModelData> CreateModelData(const FString& FileType, TConstArrayView<uint8> FileData, const TMap<FString, TConstArrayView<uint8>>& AdditionalFileData, const FGuid& FileId, const ITargetPlatform* TargetPlatform) override;
 	virtual FString GetModelDataIdentifier(const FString& FileType, TConstArrayView<uint8> FileData, const TMap<FString, TConstArrayView<uint8>>& AdditionalFileData, const FGuid& FileId, const ITargetPlatform* TargetPlatform) const override;
 
-	virtual bool CanCreateModelCPU(TObjectPtr<UNNEModelData> ModelData) const override;
-	virtual TSharedPtr<UE::NNE::IModelCPU> CreateModelCPU(TObjectPtr<UNNEModelData> ModelData) override;
+	virtual bool CanCreateModelCPU(const TObjectPtr<UNNEModelData> ModelData) const override;
+	virtual TSharedPtr<UE::NNE::IModelCPU> CreateModelCPU(const TObjectPtr<UNNEModelData> ModelData) override;
 };
\ No newline at end of file
diff --git a/Engine/Source/Runtime/NNE/Public/NNERuntimeCPU.h b/Engine/Source/Runtime/NNE/Public/NNERuntimeCPU.h
--- a/Engine/Source/Runtime/NNE/Public/NNERuntimeCPU.h
+++ b/Engine/Source/Runtime/NNE/Public/NNERuntimeCPU.h
@@ -147,7 +147,7 @@ public:
 	 * @param ModelData The model data for which to create a model.
 	 * @return True if the runtime is able to create the model, false otherwise.
 	 */
-	virtual bool CanCreateModelCPU(TObjectPtr<UNNEModelData> ModelData) const = 0;
+	virtual bool CanCreateModelCPU(const TObjectPtr<UNNEModelData> ModelData) const = 0;
 	
 	/**
 	 * Create a model given some ModelData.
@@ -159,5 +159,5 @@ public:
 	 * @param ModelData The model data for which to create a model.
 	 * @return A caller owned model representing the neural network created from ModelData.
 	 */
-	virtual TSharedPtr<UE::NNE::IModelCPU> CreateModelCPU(TObjectPtr<UNNEModelData> ModelData) = 0;
+	virtual TSharedPtr<UE::NNE::IModelCPU> CreateModelCPU(const TObjectPtr<UNNEModelData> ModelData) = 0;
 };
\ No newline at end of file
diff --git a/Engine/Source/Runtime/NNE/Public/NNERuntimeGPU.h b/Engine/Source/Runtime/NNE/Public/NNERuntimeGPU.h
--- a/Engine/Source/Runtime/NNE/Public/NNERuntimeGPU.h
+++ b/Engine/Source/Runtime/NNE/Public/NNERuntimeGPU.h
@@ -147,7 +147,7 @@ public:
 	 * @param ModelData The model data for which to create a model.
 	 * @return True if the runtime is able to create the model, false otherwise.
 	 */
-	virtual bool CanCreateModelGPU(TObjectPtr<UNNEModelData> ModelData) const = 0;
+	virtual bool CanCreateModelGPU(const TObjectPtr<UNNEModelData> ModelData) const = 0;
 	
 	/**
 	 * Create a model given some ModelData.
@@ -159,5 +159,5 @@ public:
 	 * @param ModelData The model data for which to create a model.
 	 * @return A caller owned model representing the neural network created from ModelData.
 	 */
-	virtual TSharedPtr<UE::NNE::IModelGPU> CreateModelGPU(TObjectPtr<UNNEModelData> ModelData) = 0;
+	virtual TSharedPtr<UE::NNE::IModelGPU> CreateModelGPU(const TObjectPtr<UNNEModelData> ModelData) = 0;
 };
\ No newline at end of file
diff --git a/Engine/Source/Runtime/NNE/Public/NNERuntimeRDG.h b/Engine/Source/Runtime/NNE/Public/NNERuntimeRDG.h
--- a/Engine/Source/Runtime/NNE/Public/NNERuntimeRDG.h
+++ b/Engine/Source/Runtime/NNE/Public/NNERuntimeRDG.h
@@ -148,7 +148,7 @@ public:
 	 * @param ModelData The model data for which to create a model.
 	 * @return True if the runtime is able to create the model, false otherwise.
 	 */
-	virtual bool CanCreateModelRDG(TObjectPtr<UNNEModelData> ModelData) const = 0;
+	virtual bool CanCreateModelRDG(const TObjectPtr<UNNEModelData> ModelData) const = 0;
 	
 	/**
 	 * Create a model given some ModelData.
@@ -160,5 +160,5 @@ public:
 	 * @param ModelData The model data for which to create a model.
 	 * @return A caller owned model representing the neural network created from ModelData.
 	 */
-	virtual TSharedPtr<UE::NNE::IModelRDG> CreateModelRDG(TObjectPtr<UNNEModelData> ModelData) = 0;
+	virtual TSharedPtr<UE::NNE::IModelRDG> CreateModelRDG(const TObjectPtr<UNNEModelData> ModelData) = 0;
 };
\ No newline at end of file
