[NNE] Return status code cleanup:
- INNERuntime::CanCreateModelData() returns ECanCreateModelDataStatus instead of bool
- INNERuntimeXYZ::CanCreateModelXYZ() return ECanCreateModelXYZDataStatus instead of bool

#tests Editor/Game (Win64) style-transfer and smoke tests

#rnx

diff --git a/Engine/Plugins/Experimental/NNERuntimeBasicCpu/Source/NNERuntimeBasicCpu/Private/NNERuntimeBasicCpu.cpp b/Engine/Plugins/Experimental/NNERuntimeBasicCpu/Source/NNERuntimeBasicCpu/Private/NNERuntimeBasicCpu.cpp
--- a/Engine/Plugins/Experimental/NNERuntimeBasicCpu/Source/NNERuntimeBasicCpu/Private/NNERuntimeBasicCpu.cpp
+++ b/Engine/Plugins/Experimental/NNERuntimeBasicCpu/Source/NNERuntimeBasicCpu/Private/NNERuntimeBasicCpu.cpp
@@ -9,39 +9,39 @@
 // various bits of data.
 const uint32 UNNERuntimeBasicCpuImpl::Alignment = 64;
 
-bool UNNERuntimeBasicCpuImpl::CanCreateModelData(const FString& FileType, TConstArrayView<uint8> FileData, const TMap<FString, TConstArrayView<uint8>>& AdditionalFileData, const FGuid& FileId, const ITargetPlatform* TargetPlatform) const
+UNNERuntimeBasicCpuImpl::ECanCreateModelDataStatus UNNERuntimeBasicCpuImpl::CanCreateModelData(const FString& FileType, TConstArrayView<uint8> FileData, const TMap<FString, TConstArrayView<uint8>>& AdditionalFileData, const FGuid& FileId, const ITargetPlatform* TargetPlatform) const
 {
 	if (FileType.Compare("ubnne", ESearchCase::IgnoreCase) != 0)
 	{
-		return false;
+		return ECanCreateModelDataStatus::Fail;
 	}
 
 	// We require at least a magic number and version number
 	if (FileData.Num() < 2 * sizeof(uint32))
 	{
-		return false;
+		return ECanCreateModelDataStatus::Fail;
 	}
 
 	// Check magic number valid
 	const uint32* FileMagicNumber = (const uint32*)&FileData[0 * sizeof(uint32)];
 	if (*FileMagicNumber != UE::NNE::RuntimeBasic::FModelCPU::ModelMagicNumber)
 	{
-		return false;
+		return ECanCreateModelDataStatus::Fail;
 	}
 
 	// Check version number valid
 	const uint32* FileVersionNumber = (const uint32*)&FileData[1 * sizeof(uint32)];
 	if (*FileVersionNumber != UE::NNE::RuntimeBasic::FModelCPU::ModelVersionNumber)
 	{
-		return false;
+		return ECanCreateModelDataStatus::Fail;
 	}
 
-	return true;
+	return ECanCreateModelDataStatus::Ok;
 }
 
 TSharedPtr<UE::NNE::FSharedModelData> UNNERuntimeBasicCpuImpl::CreateModelData(const FString& FileType, TConstArrayView<uint8> FileData, const TMap<FString, TConstArrayView<uint8>>& AdditionalFileData, const FGuid& FileId, const ITargetPlatform* TargetPlatform)
 {
-	if (!CanCreateModelData(FileType, FileData, AdditionalFileData, FileId, TargetPlatform))
+	if (CanCreateModelData(FileType, FileData, AdditionalFileData, FileId, TargetPlatform) != ECanCreateModelDataStatus::Ok)
 	{
 		return nullptr;
 	}
@@ -61,7 +61,7 @@ FString UNNERuntimeBasicCpuImpl::GetModelDataIdentifier(const FString& FileType,
 	return FileId.ToString(EGuidFormats::Digits) + "-" + FString::FromInt(UE::NNE::RuntimeBasic::FModelCPU::ModelMagicNumber);
 }
 
-bool UNNERuntimeBasicCpuImpl::CanCreateModelCPU(const TObjectPtr<UNNEModelData> ModelData) const
+UNNERuntimeBasicCpuImpl::ECanCreateModelCPUStatus UNNERuntimeBasicCpuImpl::CanCreateModelCPU(const TObjectPtr<UNNEModelData> ModelData) const
 {
 	check(ModelData != nullptr);
 
@@ -69,7 +69,7 @@ bool UNNERuntimeBasicCpuImpl::CanCreateModelCPU(const TObjectPtr<UNNEModelData>
 
 	if (!SharedData.IsValid())
 	{
-		return false;
+		return ECanCreateModelCPUStatus::Fail;
 	}
 
 	TConstArrayView<uint8> Data = SharedData->GetView();
@@ -77,31 +77,31 @@ bool UNNERuntimeBasicCpuImpl::CanCreateModelCPU(const TObjectPtr<UNNEModelData>
 	// We require at least a magic number and version number
 	if (Data.Num() < 2 * sizeof(uint32))
 	{
-		return false;
+		return ECanCreateModelCPUStatus::Fail;
 	}
 
 	// Check magic number valid
 	const uint32* FileMagicNumber = (const uint32*)&Data[0 * sizeof(uint32)];
 	if (*FileMagicNumber != UE::NNE::RuntimeBasic::FModelCPU::ModelMagicNumber)
 	{
-		return false;
+		return ECanCreateModelCPUStatus::Fail;
 	}
 
 	// Check version number valid
 	const uint32* FileVersionNumber = (const uint32*)&Data[1 * sizeof(uint32)];
 	if (*FileVersionNumber != UE::NNE::RuntimeBasic::FModelCPU::ModelVersionNumber)
 	{
-		return false;
+		return ECanCreateModelCPUStatus::Fail;
 	}
 
-	return true;
+	return ECanCreateModelCPUStatus::Ok;
 }
 
 TSharedPtr<UE::NNE::IModelCPU> UNNERuntimeBasicCpuImpl::CreateModelCPU(const TObjectPtr<UNNEModelData> ModelData)
 {
 	check(ModelData != nullptr);
 
-	if (!CanCreateModelCPU(ModelData))
+	if (CanCreateModelCPU(ModelData) != ECanCreateModelCPUStatus::Ok)
 	{
 		return nullptr;
 	}
diff --git a/Engine/Plugins/Experimental/NNERuntimeBasicCpu/Source/NNERuntimeBasicCpu/Private/NNERuntimeBasicCpu.h b/Engine/Plugins/Experimental/NNERuntimeBasicCpu/Source/NNERuntimeBasicCpu/Private/NNERuntimeBasicCpu.h
--- a/Engine/Plugins/Experimental/NNERuntimeBasicCpu/Source/NNERuntimeBasicCpu/Private/NNERuntimeBasicCpu.h
+++ b/Engine/Plugins/Experimental/NNERuntimeBasicCpu/Source/NNERuntimeBasicCpu/Private/NNERuntimeBasicCpu.h
@@ -25,11 +25,11 @@ public:
 
 	virtual FString GetRuntimeName() const override { return TEXT("NNERuntimeBasicCpu"); };
 
-	virtual bool CanCreateModelData(const FString& FileType, TConstArrayView<uint8> FileData, const TMap<FString, TConstArrayView<uint8>>& AdditionalFileData, const FGuid& FileId, const ITargetPlatform* TargetPlatform) const override;
+	virtual ECanCreateModelDataStatus CanCreateModelData(const FString& FileType, TConstArrayView<uint8> FileData, const TMap<FString, TConstArrayView<uint8>>& AdditionalFileData, const FGuid& FileId, const ITargetPlatform* TargetPlatform) const override;
 	virtual TSharedPtr<UE::NNE::FSharedModelData> CreateModelData(const FString& FileType, TConstArrayView<uint8> FileData, const TMap<FString, TConstArrayView<uint8>>& AdditionalFileData, const FGuid& FileId, const ITargetPlatform* TargetPlatform) override;
 	virtual FString GetModelDataIdentifier(const FString& FileType, TConstArrayView<uint8> FileData, const TMap<FString, TConstArrayView<uint8>>& AdditionalFileData, const FGuid& FileId, const ITargetPlatform* TargetPlatform) const override;
 
-	virtual bool CanCreateModelCPU(const TObjectPtr<UNNEModelData> ModelData) const override;
+	virtual ECanCreateModelCPUStatus CanCreateModelCPU(const TObjectPtr<UNNEModelData> ModelData) const override;
 	virtual TSharedPtr<UE::NNE::IModelCPU> CreateModelCPU(const TObjectPtr<UNNEModelData> ModelData) override;
 
 private:
diff --git a/Engine/Plugins/Experimental/NNERuntimeIREE/Source/NNERuntimeIREE/Private/NNERuntimeIREE.cpp b/Engine/Plugins/Experimental/NNERuntimeIREE/Source/NNERuntimeIREE/Private/NNERuntimeIREE.cpp
--- a/Engine/Plugins/Experimental/NNERuntimeIREE/Source/NNERuntimeIREE/Private/NNERuntimeIREE.cpp
+++ b/Engine/Plugins/Experimental/NNERuntimeIREE/Source/NNERuntimeIREE/Private/NNERuntimeIREE.cpp
@@ -105,12 +105,12 @@ FString UNNERuntimeIREECpu::GetRuntimeName() const
 	return TEXT("NNERuntimeIREECpu");
 }
 
-bool UNNERuntimeIREECpu::CanCreateModelData(const FString& FileType, TConstArrayView<uint8> FileData, const TMap<FString, TConstArrayView<uint8>>& AdditionalFileData, const FGuid& FileId, const ITargetPlatform* TargetPlatform) const
+UNNERuntimeIREECpu::ECanCreateModelDataStatus UNNERuntimeIREECpu::CanCreateModelData(const FString& FileType, TConstArrayView<uint8> FileData, const TMap<FString, TConstArrayView<uint8>>& AdditionalFileData, const FGuid& FileId, const ITargetPlatform* TargetPlatform) const
 {
 #if WITH_EDITOR
-	return 	FileType.Compare(TEXT("mlir"), ESearchCase::IgnoreCase) == 0;
+	return FileType.Compare(TEXT("mlir"), ESearchCase::IgnoreCase) == 0 ? ECanCreateModelDataStatus::Ok : ECanCreateModelDataStatus::Fail;
 #else
-	return false;
+	return ECanCreateModelDataStatus::Fail;
 #endif // WITH_EDITOR
 }
 
@@ -120,7 +120,7 @@ TSharedPtr<UE::NNE::FSharedModelData> UNNERuntimeIREECpu::CreateModelData(const
 	using namespace UE::NNERuntimeIREE::CPU::Private;
 
 	FString TargetPlatformName = TargetPlatform ? TargetPlatform->IniPlatformName() : UGameplayStatics::GetPlatformName();
-	if (!CanCreateModelData(FileType, FileData, AdditionalFileData, FileId, TargetPlatform))
+	if (CanCreateModelData(FileType, FileData, AdditionalFileData, FileId, TargetPlatform) != ECanCreateModelDataStatus::Ok)
 	{
 		UE_LOG(LogTemp, Warning, TEXT("UNNERuntimeIREECpu cannot create the model data with id %s (Filetype: %s) for platform %s"), *FileId.ToString(EGuidFormats::Digits).ToLower(), *FileType, *TargetPlatformName);
 		return TSharedPtr<UE::NNE::FSharedModelData>();
@@ -222,14 +222,14 @@ FString UNNERuntimeIREECpu::GetModelDataIdentifier(const FString& FileType, TCon
 	return UE::NNERuntimeIREE::CPU::Private::GetModelDataIdentifier(GetRuntimeName(), UNNERuntimeIREECpu::GUID, FileId.ToString(EGuidFormats::Digits), PlatformName, "");
 }
 
-bool UNNERuntimeIREECpu::CanCreateModelCPU(const TObjectPtr<UNNEModelData> ModelData) const
+UNNERuntimeIREECpu::ECanCreateModelCPUStatus UNNERuntimeIREECpu::CanCreateModelCPU(const TObjectPtr<UNNEModelData> ModelData) const
 {
 	check(ModelData != nullptr);
 
 	TSharedPtr<UE::NNE::FSharedModelData> SharedData = ModelData->GetModelData(GetRuntimeName());
 	if (!SharedData.IsValid())
 	{
-		return false;
+		return ECanCreateModelCPUStatus::Fail;
 	}
 
 	TConstArrayView<uint8> SharedDataView = SharedData->GetView();
@@ -237,12 +237,13 @@ bool UNNERuntimeIREECpu::CanCreateModelCPU(const TObjectPtr<UNNEModelData> Model
 	int32 VersionSize = sizeof(UNNERuntimeIREECpu::Version);
 	if (SharedDataView.Num() <= GuidSize + VersionSize)
 	{
-		return false;
+		return ECanCreateModelCPUStatus::Fail;
 	}
 
 	bool bResult = FGenericPlatformMemory::Memcmp(&(SharedDataView[0]), &(UNNERuntimeIREECpu::GUID), GuidSize) == 0;
 	bResult &= FGenericPlatformMemory::Memcmp(&(SharedDataView[GuidSize]), &(UNNERuntimeIREECpu::Version), VersionSize) == 0;
-	return bResult;
+
+	return bResult ? ECanCreateModelCPUStatus::Ok : ECanCreateModelCPUStatus::Fail;
 }
 
 TSharedPtr<UE::NNE::IModelCPU> UNNERuntimeIREECpu::CreateModelCPU(const TObjectPtr<UNNEModelData> ModelData)
@@ -251,7 +252,7 @@ TSharedPtr<UE::NNE::IModelCPU> UNNERuntimeIREECpu::CreateModelCPU(const TObjectP
 
 	using namespace UE::NNERuntimeIREE::CPU::Private;
 
-	if (!CanCreateModelCPU(ModelData))
+	if (CanCreateModelCPU(ModelData) != ECanCreateModelCPUStatus::Ok)
 	{
 		UE_LOG(LogTemp, Warning, TEXT("UNNERuntimeIREECpu cannot create a model from the model data with id %s"), *ModelData->GetFileId().ToString(EGuidFormats::Digits));
 		return TSharedPtr<UE::NNE::IModelCPU>();
@@ -414,12 +415,12 @@ FString UNNERuntimeIREEGpu::GetRuntimeName() const
 	return TEXT("");
 }
 
-bool UNNERuntimeIREEGpu::CanCreateModelData(const FString& FileType, TConstArrayView<uint8> FileData, const TMap<FString, TConstArrayView<uint8>>& AdditionalFileData, const FGuid& FileId, const ITargetPlatform* TargetPlatform) const
+UNNERuntimeIREEGpu::ECanCreateModelDataStatus UNNERuntimeIREEGpu::CanCreateModelData(const FString& FileType, TConstArrayView<uint8> FileData, const TMap<FString, TConstArrayView<uint8>>& AdditionalFileData, const FGuid& FileId, const ITargetPlatform* TargetPlatform) const
 {
 #if WITH_EDITOR
-	return	FileType.Compare(TEXT("mlir"), ESearchCase::IgnoreCase) == 0;
+	return FileType.Compare(TEXT("mlir"), ESearchCase::IgnoreCase) == 0 ? ECanCreateModelDataStatus::Ok : ECanCreateModelDataStatus::Fail;
 #else
-	return false;
+	return ECanCreateModelDataStatus::Fail;
 #endif // WITH_EDITOR
 }
 
@@ -434,14 +435,14 @@ FString UNNERuntimeIREEGpu::GetModelDataIdentifier(const FString& FileType, TCon
 	return UE::NNERuntimeIREE::CPU::Private::GetModelDataIdentifier(GetRuntimeName(), GetGUID(), FileId.ToString(EGuidFormats::Digits), PlatformName, "");
 }
 
-bool UNNERuntimeIREEGpu::CanCreateModelGPU(const TObjectPtr<UNNEModelData> ModelData) const
+UNNERuntimeIREEGpu::ECanCreateModelGPUStatus UNNERuntimeIREEGpu::CanCreateModelGPU(const TObjectPtr<UNNEModelData> ModelData) const
 {
 	check(ModelData != nullptr);
 
 	TSharedPtr<UE::NNE::FSharedModelData> SharedData = ModelData->GetModelData(GetRuntimeName());
 	if (!SharedData.IsValid())
 	{
-		return false;
+		return ECanCreateModelGPUStatus::Fail;
 	}
 
 	TConstArrayView<uint8> SharedDataView = SharedData->GetView();
@@ -451,19 +452,20 @@ bool UNNERuntimeIREEGpu::CanCreateModelGPU(const TObjectPtr<UNNEModelData> Model
 	int32 VersionSize = sizeof(Version);
 	if (SharedDataView.Num() <= GuidSize + VersionSize)
 	{
-		return false;
+		return ECanCreateModelGPUStatus::Fail;
 	}
 
 	bool bResult = FGenericPlatformMemory::Memcmp(&(SharedDataView[0]), &(Guid), GuidSize) == 0;
 	bResult &= FGenericPlatformMemory::Memcmp(&(SharedDataView[GuidSize]), &(Version), VersionSize) == 0;
-	return bResult;
+
+	return bResult ? ECanCreateModelGPUStatus::Ok : ECanCreateModelGPUStatus::Fail;
 }
 
 TSharedPtr<UE::NNE::IModelGPU> UNNERuntimeIREEGpu::CreateModelGPU(const TObjectPtr<UNNEModelData> ModelData)
 {
 	check(ModelData != nullptr);
 
-	if (!CanCreateModelGPU(ModelData))
+	if (CanCreateModelGPU(ModelData) != ECanCreateModelGPUStatus::Ok)
 	{
 		return TSharedPtr<UE::NNE::IModelGPU>();
 	}
@@ -555,12 +557,12 @@ FString UNNERuntimeIREERdg::GetRuntimeName() const
 	return TEXT("NNERuntimeIREERdg");
 }
 
-bool UNNERuntimeIREERdg::CanCreateModelData(const FString& FileType, TConstArrayView<uint8> FileData, const TMap<FString, TConstArrayView<uint8>>& AdditionalFileData, const FGuid& FileId, const ITargetPlatform* TargetPlatform) const
+UNNERuntimeIREERdg::ECanCreateModelDataStatus UNNERuntimeIREERdg::CanCreateModelData(const FString& FileType, TConstArrayView<uint8> FileData, const TMap<FString, TConstArrayView<uint8>>& AdditionalFileData, const FGuid& FileId, const ITargetPlatform* TargetPlatform) const
 {
 #if WITH_EDITOR
-	return	FileType.Compare(TEXT("mlir"), ESearchCase::IgnoreCase) == 0;
+	return	FileType.Compare(TEXT("mlir"), ESearchCase::IgnoreCase) == 0 ? ECanCreateModelDataStatus::Ok : ECanCreateModelDataStatus::Fail;
 #else
-	return false;
+	return ECanCreateModelDataStatus::Fail;
 #endif // WITH_EDITOR
 }
 
@@ -575,14 +577,14 @@ FString UNNERuntimeIREERdg::GetModelDataIdentifier(const FString& FileType, TCon
 	return UE::NNERuntimeIREE::CPU::Private::GetModelDataIdentifier(GetRuntimeName(), UNNERuntimeIREERdg::GUID, FileId.ToString(EGuidFormats::Digits), PlatformName, "");
 }
 
-bool UNNERuntimeIREERdg::CanCreateModelRDG(const TObjectPtr<UNNEModelData> ModelData) const
+UNNERuntimeIREERdg::ECanCreateModelRDGStatus UNNERuntimeIREERdg::CanCreateModelRDG(const TObjectPtr<UNNEModelData> ModelData) const
 {
 	check(ModelData != nullptr);
 
 	TSharedPtr<UE::NNE::FSharedModelData> SharedData = ModelData->GetModelData(GetRuntimeName());
 	if (!SharedData.IsValid())
 	{
-		return false;
+		return ECanCreateModelRDGStatus::Fail;
 	}
 
 	TConstArrayView<uint8> SharedDataView = SharedData->GetView();
@@ -590,19 +592,20 @@ bool UNNERuntimeIREERdg::CanCreateModelRDG(const TObjectPtr<UNNEModelData> Model
 	int32 VersionSize = sizeof(UNNERuntimeIREERdg::Version);
 	if (SharedDataView.Num() <= GuidSize + VersionSize)
 	{
-		return false;
+		return ECanCreateModelRDGStatus::Fail;
 	}
 
 	bool bResult = FGenericPlatformMemory::Memcmp(&(SharedDataView[0]), &(UNNERuntimeIREERdg::GUID), GuidSize) == 0;
 	bResult &= FGenericPlatformMemory::Memcmp(&(SharedDataView[GuidSize]), &(UNNERuntimeIREERdg::Version), VersionSize) == 0;
-	return bResult;
+
+	return bResult ? ECanCreateModelRDGStatus::Ok : ECanCreateModelRDGStatus::Fail;
 }
 
 TSharedPtr<UE::NNE::IModelRDG> UNNERuntimeIREERdg::CreateModelRDG(const TObjectPtr<UNNEModelData> ModelData)
 {
 	check(ModelData != nullptr);
 
-	if (!CanCreateModelRDG(ModelData))
+	if (CanCreateModelRDG(ModelData) != ECanCreateModelRDGStatus::Ok)
 	{
 		return TSharedPtr<UE::NNE::IModelRDG>();
 	}
diff --git a/Engine/Plugins/Experimental/NNERuntimeIREE/Source/NNERuntimeIREE/Private/NNERuntimeIREE.h b/Engine/Plugins/Experimental/NNERuntimeIREE/Source/NNERuntimeIREE/Private/NNERuntimeIREE.h
--- a/Engine/Plugins/Experimental/NNERuntimeIREE/Source/NNERuntimeIREE/Private/NNERuntimeIREE.h
+++ b/Engine/Plugins/Experimental/NNERuntimeIREE/Source/NNERuntimeIREE/Private/NNERuntimeIREE.h
@@ -24,13 +24,13 @@ public:
 
 	//~ Begin INNERuntime Interface
 	virtual FString GetRuntimeName() const override;
-	virtual bool CanCreateModelData(const FString& FileType, TConstArrayView<uint8> FileData, const TMap<FString, TConstArrayView<uint8>>& AdditionalFileData, const FGuid& FileId, const ITargetPlatform* TargetPlatform) const override;
+	virtual ECanCreateModelDataStatus CanCreateModelData(const FString& FileType, TConstArrayView<uint8> FileData, const TMap<FString, TConstArrayView<uint8>>& AdditionalFileData, const FGuid& FileId, const ITargetPlatform* TargetPlatform) const override;
 	virtual TSharedPtr<UE::NNE::FSharedModelData> CreateModelData(const FString& FileType, TConstArrayView<uint8> FileData, const TMap<FString, TConstArrayView<uint8>>& AdditionalFileData, const FGuid& FileId, const ITargetPlatform* TargetPlatform) override;
 	virtual FString GetModelDataIdentifier(const FString& FileType, TConstArrayView<uint8> FileData, const TMap<FString, TConstArrayView<uint8>>& AdditionalFileData, const FGuid& FileId, const ITargetPlatform* TargetPlatform) const override;
 	//~ End INNERuntime Interface
 
 	//~ Begin INNERuntimeCPU Interface
-	virtual bool CanCreateModelCPU(const TObjectPtr<UNNEModelData> ModelData) const override;
+	virtual ECanCreateModelCPUStatus CanCreateModelCPU(const TObjectPtr<UNNEModelData> ModelData) const override;
 	virtual TSharedPtr<UE::NNE::IModelCPU> CreateModelCPU(const TObjectPtr<UNNEModelData> ModelData) override;
 	//~ End INNERuntimeCPU Interface
 
@@ -38,13 +38,13 @@ public:
 #else
 	//~ Begin INNERuntime Interface
 	virtual FString GetRuntimeName() const override { return ""; };
-	virtual bool CanCreateModelData(const FString& FileType, TConstArrayView<uint8> FileData, const TMap<FString, TConstArrayView<uint8>>& AdditionalFileData, const FGuid& FileId, const ITargetPlatform* TargetPlatform) const override { return false; };
+	virtual ECanCreateModelDataStatus CanCreateModelData(const FString& FileType, TConstArrayView<uint8> FileData, const TMap<FString, TConstArrayView<uint8>>& AdditionalFileData, const FGuid& FileId, const ITargetPlatform* TargetPlatform) const override { return ECanCreateModelDataStatus::Fail; };
 	virtual TSharedPtr<UE::NNE::FSharedModelData> CreateModelData(const FString& FileType, TConstArrayView<uint8> FileData, const TMap<FString, TConstArrayView<uint8>>& AdditionalFileData, const FGuid& FileId, const ITargetPlatform* TargetPlatform) override { return TSharedPtr<UE::NNE::FSharedModelData>(); };
 	virtual FString GetModelDataIdentifier(const FString& FileType, TConstArrayView<uint8> FileData, const TMap<FString, TConstArrayView<uint8>>& AdditionalFileData, const FGuid& FileId, const ITargetPlatform* TargetPlatform) const override { return ""; };
 	//~ End INNERuntime Interface
 
 	//~ Begin INNERuntimeCPU Interface
-	virtual bool CanCreateModelCPU(const TObjectPtr<UNNEModelData> ModelData) const override { return false; };
+	virtual ECanCreateModelCPUStatus CanCreateModelCPU(const TObjectPtr<UNNEModelData> ModelData) const override { return ECanCreateModelCPUStatus::Fail; };
 	virtual TSharedPtr<UE::NNE::IModelCPU> CreateModelCPU(const TObjectPtr<UNNEModelData> ModelData) override { return TSharedPtr<UE::NNE::IModelCPU>(); };
 	//~ End INNERuntimeCPU Interface
 #endif // WITH_NNE_RUNTIME_IREE
@@ -59,13 +59,13 @@ public:
 #ifdef WITH_NNE_RUNTIME_IREE
 	//~ Begin INNERuntime Interface
 	virtual FString GetRuntimeName() const override;
-	virtual bool CanCreateModelData(const FString& FileType, TConstArrayView<uint8> FileData, const TMap<FString, TConstArrayView<uint8>>& AdditionalFileData, const FGuid& FileId, const ITargetPlatform* TargetPlatform) const override;
+	virtual ECanCreateModelDataStatus CanCreateModelData(const FString& FileType, TConstArrayView<uint8> FileData, const TMap<FString, TConstArrayView<uint8>>& AdditionalFileData, const FGuid& FileId, const ITargetPlatform* TargetPlatform) const override;
 	virtual TSharedPtr<UE::NNE::FSharedModelData> CreateModelData(const FString& FileType, TConstArrayView<uint8> FileData, const TMap<FString, TConstArrayView<uint8>>& AdditionalFileData, const FGuid& FileId, const ITargetPlatform* TargetPlatform) override;
 	virtual FString GetModelDataIdentifier(const FString& FileType, TConstArrayView<uint8> FileData, const TMap<FString, TConstArrayView<uint8>>& AdditionalFileData, const FGuid& FileId, const ITargetPlatform* TargetPlatform) const override;
 	//~ End INNERuntime Interface
 
 	//~ Begin INNERuntimeGPU Interface
-	virtual bool CanCreateModelGPU(const TObjectPtr<UNNEModelData> ModelData) const override;
+	virtual ECanCreateModelGPUStatus CanCreateModelGPU(const TObjectPtr<UNNEModelData> ModelData) const override;
 	virtual TSharedPtr<UE::NNE::IModelGPU> CreateModelGPU(const TObjectPtr<UNNEModelData> ModelData) override;
 	//~ End INNERuntimeGPU Interface
 
@@ -75,13 +75,13 @@ public:
 #else
 	//~ Begin INNERuntime Interface
 	virtual FString GetRuntimeName() const override { return ""; };
-	virtual bool CanCreateModelData(const FString& FileType, TConstArrayView<uint8> FileData, const TMap<FString, TConstArrayView<uint8>>& AdditionalFileData, const FGuid& FileId, const ITargetPlatform* TargetPlatform) const override { return false; };
+	virtual ECanCreateModelDataStatus CanCreateModelData(const FString& FileType, TConstArrayView<uint8> FileData, const TMap<FString, TConstArrayView<uint8>>& AdditionalFileData, const FGuid& FileId, const ITargetPlatform* TargetPlatform) const override { return ECanCreateModelDataStatus::Fail; };
 	virtual TSharedPtr<UE::NNE::FSharedModelData> CreateModelData(const FString& FileType, TConstArrayView<uint8> FileData, const TMap<FString, TConstArrayView<uint8>>& AdditionalFileData, const FGuid& FileId, const ITargetPlatform* TargetPlatform) override { return TSharedPtr<UE::NNE::FSharedModelData>(); };
 	virtual FString GetModelDataIdentifier(const FString& FileType, TConstArrayView<uint8> FileData, const TMap<FString, TConstArrayView<uint8>>& AdditionalFileData, const FGuid& FileId, const ITargetPlatform* TargetPlatform) const override { return ""; };
 	//~ End INNERuntime Interface
 
 	//~ Begin INNERuntimeGPU Interface
-	virtual bool CanCreateModelGPU(const TObjectPtr<UNNEModelData> ModelData) const override { return false; };
+	virtual ECanCreateModelGPUStatus CanCreateModelGPU(const TObjectPtr<UNNEModelData> ModelData) const override { return ECanCreateModelGPUStatus::Fail; };
 	virtual TSharedPtr<UE::NNE::IModelGPU> CreateModelGPU(const TObjectPtr<UNNEModelData> ModelData) override { return TSharedPtr<UE::NNE::IModelGPU>(); };
 	//~ End INNERuntimeGPU Interface
 #endif // WITH_NNE_RUNTIME_IREE
@@ -143,13 +143,13 @@ public:
 
 	//~ Begin INNERuntime Interface
 	virtual FString GetRuntimeName() const override;
-	virtual bool CanCreateModelData(const FString& FileType, TConstArrayView<uint8> FileData, const TMap<FString, TConstArrayView<uint8>>& AdditionalFileData, const FGuid& FileId, const ITargetPlatform* TargetPlatform) const override;
+	virtual ECanCreateModelDataStatus CanCreateModelData(const FString& FileType, TConstArrayView<uint8> FileData, const TMap<FString, TConstArrayView<uint8>>& AdditionalFileData, const FGuid& FileId, const ITargetPlatform* TargetPlatform) const override;
 	virtual TSharedPtr<UE::NNE::FSharedModelData> CreateModelData(const FString& FileType, TConstArrayView<uint8> FileData, const TMap<FString, TConstArrayView<uint8>>& AdditionalFileData, const FGuid& FileId, const ITargetPlatform* TargetPlatform) override;
 	virtual FString GetModelDataIdentifier(const FString& FileType, TConstArrayView<uint8> FileData, const TMap<FString, TConstArrayView<uint8>>& AdditionalFileData, const FGuid& FileId, const ITargetPlatform* TargetPlatform) const override;
 	//~ End INNERuntime Interface
 
 	//~ Begin INNERuntimeRdg Interface
-	virtual bool CanCreateModelRDG(const TObjectPtr<UNNEModelData> ModelData) const override;
+	virtual ECanCreateModelRDGStatus CanCreateModelRDG(const TObjectPtr<UNNEModelData> ModelData) const override;
 	virtual TSharedPtr<UE::NNE::IModelRDG> CreateModelRDG(const TObjectPtr<UNNEModelData> ModelData) override;
 	//~ End INNERuntimeRdg Interface
 
@@ -157,13 +157,13 @@ public:
 #else
 	//~ Begin INNERuntime Interface
 	virtual FString GetRuntimeName() const override { return ""; };
-	virtual bool CanCreateModelData(const FString& FileType, TConstArrayView<uint8> FileData, const TMap<FString, TConstArrayView<uint8>>& AdditionalFileData, const FGuid& FileId, const ITargetPlatform* TargetPlatform) const override { return false; };
+	virtual ECanCreateModelDataStatus CanCreateModelData(const FString& FileType, TConstArrayView<uint8> FileData, const TMap<FString, TConstArrayView<uint8>>& AdditionalFileData, const FGuid& FileId, const ITargetPlatform* TargetPlatform) const override { return ECanCreateModelDataStatus::Fail; };
 	virtual TSharedPtr<UE::NNE::FSharedModelData> CreateModelData(const FString& FileType, TConstArrayView<uint8> FileData, const TMap<FString, TConstArrayView<uint8>>& AdditionalFileData, const FGuid& FileId, const ITargetPlatform* TargetPlatform) override { return TSharedPtr<UE::NNE::FSharedModelData>(); };
 	virtual FString GetModelDataIdentifier(const FString& FileType, TConstArrayView<uint8> FileData, const TMap<FString, TConstArrayView<uint8>>& AdditionalFileData, const FGuid& FileId, const ITargetPlatform* TargetPlatform) const override { return ""; };
 	//~ End INNERuntime Interface
 
 	//~ Begin INNERuntimeRDG Interface
-	virtual bool CanCreateModelRDG(const TObjectPtr<UNNEModelData> ModelData) const override { return false; };
+	virtual ECanCreateModelRDGStatus CanCreateModelRDG(const TObjectPtr<UNNEModelData> ModelData) const override { return ECanCreateModelRDGStatus::Fail; };
 	virtual TSharedPtr<UE::NNE::IModelRDG> CreateModelRDG(const TObjectPtr<UNNEModelData> ModelData) override { return TSharedPtr<UE::NNE::IModelRDG>(); };
 	//~ End INNERuntimeRDG Interface
 #endif // WITH_NNE_RUNTIME_IREE
diff --git a/Engine/Plugins/Experimental/NNERuntimeRDG/Source/NNERuntimeRDG/Private/NNERuntimeRDGDml.cpp b/Engine/Plugins/Experimental/NNERuntimeRDG/Source/NNERuntimeRDG/Private/NNERuntimeRDGDml.cpp
--- a/Engine/Plugins/Experimental/NNERuntimeRDG/Source/NNERuntimeRDG/Private/NNERuntimeRDGDml.cpp
+++ b/Engine/Plugins/Experimental/NNERuntimeRDG/Source/NNERuntimeRDG/Private/NNERuntimeRDGDml.cpp
@@ -111,19 +111,19 @@ bool FRuntimeDmlStartup()
 
 using namespace UE::NNERuntimeRDG::Private::Dml;
 
-bool UNNERuntimeRDGDmlImpl::CanCreateModelData(const FString& FileType, TConstArrayView<uint8> FileData, const TMap<FString, TConstArrayView<uint8>>& AdditionalFileData, const FGuid& FileId, const ITargetPlatform* TargetPlatform) const
+UNNERuntimeRDGDmlImpl::ECanCreateModelDataStatus UNNERuntimeRDGDmlImpl::CanCreateModelData(const FString& FileType, TConstArrayView<uint8> FileData, const TMap<FString, TConstArrayView<uint8>>& AdditionalFileData, const FGuid& FileId, const ITargetPlatform* TargetPlatform) const
 {
 #ifdef NNE_UTILITIES_AVAILABLE
-	return FileType.Compare("onnx", ESearchCase::IgnoreCase) == 0;
+	return FileType.Compare("onnx", ESearchCase::IgnoreCase) == 0 ? ECanCreateModelDataStatus::Ok : ECanCreateModelDataStatus::Fail;
 #else
 	UE_LOG(LogNNE, Display, TEXT("NNEUtilities is not available on this platform"));
-	return false;
+	return ECanCreateModelDataStatus::Fail;
 #endif
 }
 
 TSharedPtr<UE::NNE::FSharedModelData> UNNERuntimeRDGDmlImpl::CreateModelData(const FString& FileType, TConstArrayView<uint8> FileData, const TMap<FString, TConstArrayView<uint8>>& AdditionalFileData, const FGuid& FileId, const ITargetPlatform* TargetPlatform)
 {
-	if (!CanCreateModelData(FileType, FileData, AdditionalFileData, FileId, TargetPlatform))
+	if (CanCreateModelData(FileType, FileData, AdditionalFileData, FileId, TargetPlatform) != ECanCreateModelDataStatus::Ok)
 	{
 		return {};
 	}
@@ -165,39 +165,40 @@ FString UNNERuntimeRDGDmlImpl::GetModelDataIdentifier(const FString& FileType, T
 	return FileId.ToString(EGuidFormats::Digits) + "-" + FModelInfo::Get()->GetGuid().ToString(EGuidFormats::Digits) + "-" + FString::FromInt(FModelInfo::Get()->GetVersion());
 }
 
-bool UNNERuntimeRDGDmlImpl::CanCreateModelRDG(const TObjectPtr<UNNEModelData> ModelData) const
+UNNERuntimeRDGDmlImpl::ECanCreateModelRDGStatus UNNERuntimeRDGDmlImpl::CanCreateModelRDG(const TObjectPtr<UNNEModelData> ModelData) const
 {
 #ifdef NNE_USE_DIRECTML
 	if (bRegisterOnlyOperators)
 	{
-		return false;
+		return ECanCreateModelRDGStatus::Fail;
 	}
 
 	TSharedPtr<UE::NNE::FSharedModelData> SharedData = ModelData->GetModelData(GetRuntimeName());
 
 	if (!SharedData.IsValid())
 	{
-		return false;
+		return ECanCreateModelRDGStatus::Fail;
 	}
 
 	TConstArrayView<uint8> Data = SharedData->GetView();
 
 	if (Data.Num() <= FModelInfo::Get()->GetGuidSize() + FModelInfo::Get()->GetVersionSize())
 	{
-		return false;
+		return ECanCreateModelRDGStatus::Fail;
 	}
 
 	bool bResult = FModelInfo::Get()->ValidateGuidAndVersion(Data.GetData(), Data.GetData() + FModelInfo::Get()->GetGuidSize());
-	return bResult;
+
+	return bResult ? ECanCreateModelRDGStatus::Ok : ECanCreateModelRDGStatus::Fail;
 #else
-	return false;
+	return ECanCreateModelRDGStatus::Fail;
 #endif
 };
 
 TSharedPtr<UE::NNE::IModelRDG> UNNERuntimeRDGDmlImpl::CreateModelRDG(const TObjectPtr<UNNEModelData> ModelData)
 {
 #ifdef NNE_USE_DIRECTML
-	if (!CanCreateModelRDG(ModelData))
+	if (CanCreateModelRDG(ModelData) != ECanCreateModelRDGStatus::Ok)
 	{
 		return TSharedPtr<UE::NNE::IModelRDG>();
 	}
diff --git a/Engine/Plugins/Experimental/NNERuntimeRDG/Source/NNERuntimeRDG/Private/NNERuntimeRDGDml.h b/Engine/Plugins/Experimental/NNERuntimeRDG/Source/NNERuntimeRDG/Private/NNERuntimeRDGDml.h
--- a/Engine/Plugins/Experimental/NNERuntimeRDG/Source/NNERuntimeRDG/Private/NNERuntimeRDGDml.h
+++ b/Engine/Plugins/Experimental/NNERuntimeRDG/Source/NNERuntimeRDG/Private/NNERuntimeRDGDml.h
@@ -33,11 +33,11 @@ public:
 
 	virtual FString GetRuntimeName() const override { return TEXT("NNERuntimeRDGDml"); };
 
-	virtual bool CanCreateModelData(const FString& FileType, TConstArrayView<uint8> FileData, const TMap<FString, TConstArrayView<uint8>>& AdditionalFileData, const FGuid& FileId, const ITargetPlatform* TargetPlatform) const override;
+	virtual ECanCreateModelDataStatus CanCreateModelData(const FString& FileType, TConstArrayView<uint8> FileData, const TMap<FString, TConstArrayView<uint8>>& AdditionalFileData, const FGuid& FileId, const ITargetPlatform* TargetPlatform) const override;
 	virtual TSharedPtr<UE::NNE::FSharedModelData> CreateModelData(const FString& FileType, TConstArrayView<uint8> FileData, const TMap<FString, TConstArrayView<uint8>>& AdditionalFileData, const FGuid& FileId, const ITargetPlatform* TargetPlatform) override;
 	virtual FString GetModelDataIdentifier(const FString& FileType, TConstArrayView<uint8> FileData, const TMap<FString, TConstArrayView<uint8>>& AdditionalFileData, const FGuid& FileId, const ITargetPlatform* TargetPlatform) const override;
 
-	virtual bool CanCreateModelRDG(const TObjectPtr<UNNEModelData> ModelData) const override;
+	virtual ECanCreateModelRDGStatus CanCreateModelRDG(const TObjectPtr<UNNEModelData> ModelData) const override;
 	virtual TSharedPtr<UE::NNE::IModelRDG> CreateModelRDG(const TObjectPtr<UNNEModelData> ModelData) override;
 
 private:
diff --git a/Engine/Plugins/Experimental/NNERuntimeRDG/Source/NNERuntimeRDG/Private/NNERuntimeRDGHlsl.cpp b/Engine/Plugins/Experimental/NNERuntimeRDG/Source/NNERuntimeRDG/Private/NNERuntimeRDGHlsl.cpp
--- a/Engine/Plugins/Experimental/NNERuntimeRDG/Source/NNERuntimeRDG/Private/NNERuntimeRDGHlsl.cpp
+++ b/Engine/Plugins/Experimental/NNERuntimeRDG/Source/NNERuntimeRDG/Private/NNERuntimeRDGHlsl.cpp
@@ -85,17 +85,17 @@ bool UNNERuntimeRDGHlslImpl::Init()
 	return true;
 }
 
-bool UNNERuntimeRDGHlslImpl::CanCreateModelData(const FString& FileType, TConstArrayView<uint8> FileData, const TMap<FString, TConstArrayView<uint8>>& AdditionalFileData, const FGuid& FileId, const ITargetPlatform* TargetPlatform) const
+UNNERuntimeRDGHlslImpl::ECanCreateModelDataStatus UNNERuntimeRDGHlslImpl::CanCreateModelData(const FString& FileType, TConstArrayView<uint8> FileData, const TMap<FString, TConstArrayView<uint8>>& AdditionalFileData, const FGuid& FileId, const ITargetPlatform* TargetPlatform) const
 {
 #ifdef NNE_UTILITIES_AVAILABLE
-	return FileType.Compare("onnx", ESearchCase::IgnoreCase) == 0;
+	return FileType.Compare("onnx", ESearchCase::IgnoreCase) == 0 ? ECanCreateModelDataStatus::Ok : ECanCreateModelDataStatus::Fail;
 #else
 	UE_LOG(LogNNE, Display, TEXT("NNEUtilities is not available on this platform"));
-	return false;
+	return ECanCreateModelDataStatus::Fail;
 #endif
 }
 
-bool UNNERuntimeRDGHlslImpl::CanCreateModelRDG(const TObjectPtr<UNNEModelData> ModelData) const
+UNNERuntimeRDGHlslImpl::ECanCreateModelRDGStatus UNNERuntimeRDGHlslImpl::CanCreateModelRDG(const TObjectPtr<UNNEModelData> ModelData) const
 {
 	int32 GuidSize = sizeof(GUID);
 	int32 VersionSize = sizeof(Version);
@@ -103,23 +103,24 @@ bool UNNERuntimeRDGHlslImpl::CanCreateModelRDG(const TObjectPtr<UNNEModelData> M
 
 	if (!SharedData.IsValid())
 	{
-		return false;
+		return ECanCreateModelRDGStatus::Fail;
 	}
 
 	TConstArrayView<uint8> Data = SharedData->GetView();
 
 	if (Data.Num() <= GuidSize + VersionSize)
 	{
-		return false;
+		return ECanCreateModelRDGStatus::Fail;
 	}
 	bool bResult = FGenericPlatformMemory::Memcmp(&(Data[0]), &(GUID), GuidSize) == 0;
 	bResult &= FGenericPlatformMemory::Memcmp(&(Data[GuidSize]), &(Version), VersionSize) == 0;
-	return bResult;
+
+	return bResult ? ECanCreateModelRDGStatus::Ok : ECanCreateModelRDGStatus::Fail;
 };
 
 TSharedPtr<UE::NNE::FSharedModelData> UNNERuntimeRDGHlslImpl::CreateModelData(const FString& FileType, TConstArrayView<uint8> FileData, const TMap<FString, TConstArrayView<uint8>>& AdditionalFileData, const FGuid& FileId, const ITargetPlatform* TargetPlatform)
 {
-	if (!CanCreateModelData(FileType, FileData, AdditionalFileData, FileId, TargetPlatform))
+	if (CanCreateModelData(FileType, FileData, AdditionalFileData, FileId, TargetPlatform) != ECanCreateModelDataStatus::Ok)
 	{
 		return {};
 	}
@@ -158,7 +159,7 @@ FString UNNERuntimeRDGHlslImpl::GetModelDataIdentifier(const FString& FileType,
 
 TSharedPtr<UE::NNE::IModelRDG> UNNERuntimeRDGHlslImpl::CreateModelRDG(const TObjectPtr<UNNEModelData> ModelData)
 {
-	if (!CanCreateModelRDG(ModelData))
+	if (CanCreateModelRDG(ModelData) != ECanCreateModelRDGStatus::Ok)
 	{
 		return TSharedPtr<UE::NNE::IModelRDG>();
 	}
diff --git a/Engine/Plugins/Experimental/NNERuntimeRDG/Source/NNERuntimeRDG/Private/NNERuntimeRDGHlsl.h b/Engine/Plugins/Experimental/NNERuntimeRDG/Source/NNERuntimeRDG/Private/NNERuntimeRDGHlsl.h
--- a/Engine/Plugins/Experimental/NNERuntimeRDG/Source/NNERuntimeRDG/Private/NNERuntimeRDGHlsl.h
+++ b/Engine/Plugins/Experimental/NNERuntimeRDG/Source/NNERuntimeRDG/Private/NNERuntimeRDGHlsl.h
@@ -28,10 +28,10 @@ public:
 
 	virtual FString GetRuntimeName() const override { return TEXT("NNERuntimeRDGHlsl"); };
 
-	virtual bool CanCreateModelData(const FString& FileType, TConstArrayView<uint8> FileData, const TMap<FString, TConstArrayView<uint8>>& AdditionalFileData, const FGuid& FileId, const ITargetPlatform* TargetPlatform) const override;
+	virtual ECanCreateModelDataStatus CanCreateModelData(const FString& FileType, TConstArrayView<uint8> FileData, const TMap<FString, TConstArrayView<uint8>>& AdditionalFileData, const FGuid& FileId, const ITargetPlatform* TargetPlatform) const override;
 	virtual TSharedPtr<UE::NNE::FSharedModelData> CreateModelData(const FString& FileType, TConstArrayView<uint8> FileData, const TMap<FString, TConstArrayView<uint8>>& AdditionalFileData, const FGuid& FileId, const ITargetPlatform* TargetPlatform) override;
 	virtual FString GetModelDataIdentifier(const FString& FileType, TConstArrayView<uint8> FileData, const TMap<FString, TConstArrayView<uint8>>& AdditionalFileData, const FGuid& FileId, const ITargetPlatform* TargetPlatform) const override;
 
-	virtual bool CanCreateModelRDG(const TObjectPtr<UNNEModelData> ModelData) const override;
+	virtual ECanCreateModelRDGStatus CanCreateModelRDG(const TObjectPtr<UNNEModelData> ModelData) const override;
 	virtual TSharedPtr<UE::NNE::IModelRDG> CreateModelRDG(const TObjectPtr<UNNEModelData> ModelData) override;
 };
\ No newline at end of file
diff --git a/Engine/Plugins/NNE/NNERuntimeORT/Source/NNERuntimeORT/Private/NNERuntimeORT.cpp b/Engine/Plugins/NNE/NNERuntimeORT/Source/NNERuntimeORT/Private/NNERuntimeORT.cpp
--- a/Engine/Plugins/NNE/NNERuntimeORT/Source/NNERuntimeORT/Private/NNERuntimeORT.cpp
+++ b/Engine/Plugins/NNE/NNERuntimeORT/Source/NNERuntimeORT/Private/NNERuntimeORT.cpp
@@ -24,19 +24,19 @@ int32 UNNERuntimeORTDml::Version = 0x00000001;
 FGuid UNNERuntimeORTCpu::GUID = FGuid((int32)'O', (int32)'C', (int32)'P', (int32)'U');
 int32 UNNERuntimeORTCpu::Version = 0x00000001;
 
-bool UNNERuntimeORTDml::CanCreateModelData(const FString& FileType, TConstArrayView<uint8> FileData, const TMap<FString, TConstArrayView<uint8>>& AdditionalFileData, const FGuid& FileId, const ITargetPlatform* TargetPlatform) const
+UNNERuntimeORTDml::ECanCreateModelDataStatus UNNERuntimeORTDml::CanCreateModelData(const FString& FileType, TConstArrayView<uint8> FileData, const TMap<FString, TConstArrayView<uint8>>& AdditionalFileData, const FGuid& FileId, const ITargetPlatform* TargetPlatform) const
 {
-	return !FileData.IsEmpty() && FileType.Compare("onnx", ESearchCase::IgnoreCase) == 0;
+	return (!FileData.IsEmpty() && FileType.Compare("onnx", ESearchCase::IgnoreCase) == 0) ? ECanCreateModelDataStatus::Ok : ECanCreateModelDataStatus::Fail;
 }
 
-bool UNNERuntimeORTCpu::CanCreateModelData(const FString& FileType, TConstArrayView<uint8> FileData, const TMap<FString, TConstArrayView<uint8>>& AdditionalFileData, const FGuid& FileId, const ITargetPlatform* TargetPlatform) const
+UNNERuntimeORTCpu::ECanCreateModelDataStatus UNNERuntimeORTCpu::CanCreateModelData(const FString& FileType, TConstArrayView<uint8> FileData, const TMap<FString, TConstArrayView<uint8>>& AdditionalFileData, const FGuid& FileId, const ITargetPlatform* TargetPlatform) const
 {
-	return !FileData.IsEmpty() && FileType.Compare("onnx", ESearchCase::IgnoreCase) == 0;
+	return (!FileData.IsEmpty() && FileType.Compare("onnx", ESearchCase::IgnoreCase) == 0) ? ECanCreateModelDataStatus::Ok : ECanCreateModelDataStatus::Fail;
 }
 
 TSharedPtr<UE::NNE::FSharedModelData> UNNERuntimeORTDml::CreateModelData(const FString& FileType, TConstArrayView<uint8> FileData, const TMap<FString, TConstArrayView<uint8>>& AdditionalFileData, const FGuid& FileId, const ITargetPlatform* TargetPlatform)
 {
-	if (!CanCreateModelData(FileType, FileData, AdditionalFileData, FileId, TargetPlatform))
+	if (CanCreateModelData(FileType, FileData, AdditionalFileData, FileId, TargetPlatform) != ECanCreateModelDataStatus::Ok)
 	{
 		return {};
 	}
@@ -64,7 +64,7 @@ TSharedPtr<UE::NNE::FSharedModelData> UNNERuntimeORTDml::CreateModelData(const F
 
 TSharedPtr<UE::NNE::FSharedModelData> UNNERuntimeORTCpu::CreateModelData(const FString& FileType, TConstArrayView<uint8> FileData, const TMap<FString, TConstArrayView<uint8>>& AdditionalFileData, const FGuid& FileId, const ITargetPlatform* TargetPlatform)
 {
-	if (!CanCreateModelData(FileType, FileData, AdditionalFileData, FileId, TargetPlatform))
+	if (CanCreateModelData(FileType, FileData, AdditionalFileData, FileId, TargetPlatform) != ECanCreateModelDataStatus::Ok)
 	{
 		return {};
 	}
@@ -122,7 +122,7 @@ FString UNNERuntimeORTCpu::GetRuntimeName() const
 	return TEXT("NNERuntimeORTCpu");
 }
 
-bool UNNERuntimeORTCpu::CanCreateModelCPU(const TObjectPtr<UNNEModelData> ModelData) const
+UNNERuntimeORTCpu::ECanCreateModelCPUStatus UNNERuntimeORTCpu::CanCreateModelCPU(const TObjectPtr<UNNEModelData> ModelData) const
 {
 	check(ModelData != nullptr);
 
@@ -132,19 +132,20 @@ bool UNNERuntimeORTCpu::CanCreateModelCPU(const TObjectPtr<UNNEModelData> ModelD
 
 	if (!SharedData.IsValid())
 	{
-		return false;
+		return ECanCreateModelCPUStatus::Fail;
 	}
 
 	TConstArrayView<uint8> Data = SharedData->GetView();
 
 	if (Data.Num() <= GuidSize + VersionSize)
 	{
-		return false;
+		return ECanCreateModelCPUStatus::Fail;
 	}
 
 	bool bResult = FGenericPlatformMemory::Memcmp(&(Data[0]), &(UNNERuntimeORTCpu::GUID), GuidSize) == 0;
 	bResult &= FGenericPlatformMemory::Memcmp(&(Data[GuidSize]), &(UNNERuntimeORTCpu::Version), VersionSize) == 0;
-	return bResult;
+
+	return bResult ? ECanCreateModelCPUStatus::Ok : ECanCreateModelCPUStatus::Fail;
 }
 
 TSharedPtr<UE::NNE::IModelCPU> UNNERuntimeORTCpu::CreateModelCPU(const TObjectPtr<UNNEModelData> ModelData)
@@ -152,7 +153,7 @@ TSharedPtr<UE::NNE::IModelCPU> UNNERuntimeORTCpu::CreateModelCPU(const TObjectPt
 	check(ModelData != nullptr);
 	check(ORTEnvironment.IsValid());
 
-	if (!CanCreateModelCPU(ModelData))
+	if (CanCreateModelCPU(ModelData) != ECanCreateModelCPUStatus::Ok)
 	{
 		return TSharedPtr<UE::NNE::IModelCPU>();
 	}
@@ -177,14 +178,14 @@ TSharedPtr<UE::NNE::IModelCPU> UNNERuntimeORTCpu::CreateModelCPU(const TObjectPt
 }
 
 #if PLATFORM_WINDOWS
-bool UNNERuntimeORTDml::CanCreateModelGPU(const TObjectPtr<UNNEModelData> ModelData) const
+UNNERuntimeORTDml::ECanCreateModelGPUStatus UNNERuntimeORTDml::CanCreateModelGPU(const TObjectPtr<UNNEModelData> ModelData) const
 {
 	check(ModelData != nullptr);
 
 	// In order to use DirectML we need D3D12
 	if (!IsRHID3D12())
 	{
-		return false;
+		return ECanCreateModelGPUStatus::Fail;
 	}
 
 	constexpr int32 GuidSize = sizeof(UNNERuntimeORTDml::GUID);
@@ -193,14 +194,14 @@ bool UNNERuntimeORTDml::CanCreateModelGPU(const TObjectPtr<UNNEModelData> ModelD
 
 	if (!SharedData.IsValid())
 	{
-		return false;
+		return ECanCreateModelGPUStatus::Fail;
 	}
 
 	TConstArrayView<uint8> Data = SharedData->GetView();
 
 	if (Data.Num() <= GuidSize + VersionSize)
 	{
-		return false;
+		return ECanCreateModelGPUStatus::Fail;
 	}
 
 	static const FGuid DeprecatedGUID = FGuid((int32)'O', (int32)'D', (int32)'M', (int32)'L');
@@ -208,7 +209,8 @@ bool UNNERuntimeORTDml::CanCreateModelGPU(const TObjectPtr<UNNEModelData> ModelD
 	bool bResult = FGenericPlatformMemory::Memcmp(&(Data[0]), &(UNNERuntimeORTDml::GUID), GuidSize) == 0;
 	bResult |= FGenericPlatformMemory::Memcmp(&(Data[0]), &(DeprecatedGUID), GuidSize) == 0;
 	bResult &= FGenericPlatformMemory::Memcmp(&(Data[GuidSize]), &(UNNERuntimeORTDml::Version), VersionSize) == 0;
-	return bResult;
+
+	return bResult ? ECanCreateModelGPUStatus::Ok : ECanCreateModelGPUStatus::Fail;
 }
 
 TSharedPtr<UE::NNE::IModelGPU> UNNERuntimeORTDml::CreateModelGPU(const TObjectPtr<UNNEModelData> ModelData)
@@ -216,7 +218,7 @@ TSharedPtr<UE::NNE::IModelGPU> UNNERuntimeORTDml::CreateModelGPU(const TObjectPt
 	check(ModelData != nullptr);
 	check(ORTEnvironment.IsValid());
 
-	if (!CanCreateModelGPU(ModelData))
+	if (CanCreateModelGPU(ModelData) != ECanCreateModelGPUStatus::Ok)
 	{
 		return TSharedPtr<UE::NNE::IModelGPU>();
 	}
@@ -242,9 +244,9 @@ TSharedPtr<UE::NNE::IModelGPU> UNNERuntimeORTDml::CreateModelGPU(const TObjectPt
 
 #else // PLATFORM_WINDOWS
 
-bool UNNERuntimeORTDml::CanCreateModelGPU(const TObjectPtr<UNNEModelData> ModelData) const
+UNNERuntimeORTDml::ECanCreateModelGPUStatus UNNERuntimeORTDml::CanCreateModelGPU(const TObjectPtr<UNNEModelData> ModelData) const
 {
-	return false;
+	return ECanCreateModelGPUStatus::Fail;
 }
 
 TSharedPtr<UE::NNE::IModelGPU> UNNERuntimeORTDml::CreateModelGPU(const TObjectPtr<UNNEModelData> ModelData)
diff --git a/Engine/Plugins/NNE/NNERuntimeORT/Source/NNERuntimeORT/Private/NNERuntimeORT.h b/Engine/Plugins/NNE/NNERuntimeORT/Source/NNERuntimeORT/Private/NNERuntimeORT.h
--- a/Engine/Plugins/NNE/NNERuntimeORT/Source/NNERuntimeORT/Private/NNERuntimeORT.h
+++ b/Engine/Plugins/NNE/NNERuntimeORT/Source/NNERuntimeORT/Private/NNERuntimeORT.h
@@ -34,11 +34,11 @@ public:
 
 	virtual FString GetRuntimeName() const override;
 
-	virtual bool CanCreateModelData(const FString& FileType, TConstArrayView<uint8> FileData, const TMap<FString, TConstArrayView<uint8>>& AdditionalFileData, const FGuid& FileId, const ITargetPlatform* TargetPlatform) const override;
+	virtual ECanCreateModelDataStatus CanCreateModelData(const FString& FileType, TConstArrayView<uint8> FileData, const TMap<FString, TConstArrayView<uint8>>& AdditionalFileData, const FGuid& FileId, const ITargetPlatform* TargetPlatform) const override;
 	virtual TSharedPtr<UE::NNE::FSharedModelData> CreateModelData(const FString& FileType, TConstArrayView<uint8> FileData, const TMap<FString, TConstArrayView<uint8>>& AdditionalFileData, const FGuid& FileId, const ITargetPlatform* TargetPlatform) override;
 	virtual FString GetModelDataIdentifier(const FString& FileType, TConstArrayView<uint8> FileData, const TMap<FString, TConstArrayView<uint8>>& AdditionalFileData, const FGuid& FileId, const ITargetPlatform* TargetPlatform) const override;
 
-	virtual bool CanCreateModelGPU(const TObjectPtr<UNNEModelData> ModelData) const override;
+	virtual ECanCreateModelGPUStatus CanCreateModelGPU(const TObjectPtr<UNNEModelData> ModelData) const override;
 	virtual TSharedPtr<UE::NNE::IModelGPU> CreateModelGPU(const TObjectPtr<UNNEModelData> ModelData) override;
 };
 
@@ -62,10 +62,10 @@ public:
 		
 	virtual FString GetRuntimeName() const override;
 	
-	virtual bool CanCreateModelData(const FString& FileType, TConstArrayView<uint8> FileData, const TMap<FString, TConstArrayView<uint8>>& AdditionalFileData, const FGuid& FileId, const ITargetPlatform* TargetPlatform) const override;
+	virtual ECanCreateModelDataStatus CanCreateModelData(const FString& FileType, TConstArrayView<uint8> FileData, const TMap<FString, TConstArrayView<uint8>>& AdditionalFileData, const FGuid& FileId, const ITargetPlatform* TargetPlatform) const override;
 	virtual TSharedPtr<UE::NNE::FSharedModelData> CreateModelData(const FString& FileType, TConstArrayView<uint8> FileData, const TMap<FString, TConstArrayView<uint8>>& AdditionalFileData, const FGuid& FileId, const ITargetPlatform* TargetPlatform) override;
 	virtual FString GetModelDataIdentifier(const FString& FileType, TConstArrayView<uint8> FileData, const TMap<FString, TConstArrayView<uint8>>& AdditionalFileData, const FGuid& FileId, const ITargetPlatform* TargetPlatform) const override;
 
-	virtual bool CanCreateModelCPU(const TObjectPtr<UNNEModelData> ModelData) const override;
+	virtual ECanCreateModelCPUStatus CanCreateModelCPU(const TObjectPtr<UNNEModelData> ModelData) const override;
 	virtual TSharedPtr<UE::NNE::IModelCPU> CreateModelCPU(const TObjectPtr<UNNEModelData> ModelData) override;
 };
\ No newline at end of file
diff --git a/Engine/Source/Runtime/NNE/Private/NNEModelData.cpp b/Engine/Source/Runtime/NNE/Private/NNEModelData.cpp
--- a/Engine/Source/Runtime/NNE/Private/NNEModelData.cpp
+++ b/Engine/Source/Runtime/NNE/Private/NNEModelData.cpp
@@ -136,7 +136,7 @@ namespace UE::NNE::ModelData
 		TWeakInterfacePtr<INNERuntime> NNERuntime = UE::NNE::GetRuntime<INNERuntime>(RuntimeName);
 		if (NNERuntime.IsValid())
 		{
-			if (NNERuntime->CanCreateModelData(FileType, FileData, AdditionalFileData, FileId, TargetPlatform))
+			if (NNERuntime->CanCreateModelData(FileType, FileData, AdditionalFileData, FileId, TargetPlatform) == INNERuntime::ECanCreateModelDataStatus::Ok)
 			{
 				return NNERuntime->CreateModelData(FileType, FileData, AdditionalFileData, FileId, TargetPlatform);
 			}
diff --git a/Engine/Source/Runtime/NNE/Public/NNERuntime.h b/Engine/Source/Runtime/NNE/Public/NNERuntime.h
--- a/Engine/Source/Runtime/NNE/Public/NNERuntime.h
+++ b/Engine/Source/Runtime/NNE/Public/NNERuntime.h
@@ -4,6 +4,7 @@
 
 #include "Containers/ContainersFwd.h"
 #include "Containers/UnrealString.h"
+#include "NNEStatus.h"
 #include "UObject/Interface.h"
 
 #include "NNERuntime.generated.h"
@@ -34,6 +35,8 @@ class NNE_API INNERuntime
 
 public:
 
+	using ECanCreateModelDataStatus = UE::NNE::EResultStatus;
+
 	/**
 	 * Get the name of the runtime.
 	 *
@@ -51,7 +54,7 @@ public:
 	 * @param TargetPlatform The Interface identifying the target platform for which the data needs to be created. A null pointer indicates the currently compiled/running platform.
 	 * @return True if the runtime is able to create model data, false otherwise.
 	 */
-	virtual bool CanCreateModelData(const FString& FileType, TConstArrayView<uint8> FileData, const TMap<FString, TConstArrayView<uint8>>& AdditionalFileData, const FGuid& FileId, const ITargetPlatform* TargetPlatform) const = 0;
+	virtual ECanCreateModelDataStatus CanCreateModelData(const FString& FileType, TConstArrayView<uint8> FileData, const TMap<FString, TConstArrayView<uint8>>& AdditionalFileData, const FGuid& FileId, const ITargetPlatform* TargetPlatform) const = 0;
 
 	/**
 	 * Create model data given some raw file data.
diff --git a/Engine/Source/Runtime/NNE/Public/NNERuntimeCPU.h b/Engine/Source/Runtime/NNE/Public/NNERuntimeCPU.h
--- a/Engine/Source/Runtime/NNE/Public/NNERuntimeCPU.h
+++ b/Engine/Source/Runtime/NNE/Public/NNERuntimeCPU.h
@@ -141,13 +141,15 @@ class NNE_API INNERuntimeCPU
 	
 public:
 
+	using ECanCreateModelCPUStatus = UE::NNE::EResultStatus;
+
 	/**
 	 * Check if the runtime is able to create a model given some ModelData.
 	 *
 	 * @param ModelData The model data for which to create a model.
 	 * @return True if the runtime is able to create the model, false otherwise.
 	 */
-	virtual bool CanCreateModelCPU(const TObjectPtr<UNNEModelData> ModelData) const = 0;
+	virtual ECanCreateModelCPUStatus CanCreateModelCPU(const TObjectPtr<UNNEModelData> ModelData) const = 0;
 	
 	/**
 	 * Create a model given some ModelData.
diff --git a/Engine/Source/Runtime/NNE/Public/NNERuntimeGPU.h b/Engine/Source/Runtime/NNE/Public/NNERuntimeGPU.h
--- a/Engine/Source/Runtime/NNE/Public/NNERuntimeGPU.h
+++ b/Engine/Source/Runtime/NNE/Public/NNERuntimeGPU.h
@@ -141,13 +141,15 @@ class NNE_API INNERuntimeGPU
 	
 public:
 
+	using ECanCreateModelGPUStatus = UE::NNE::EResultStatus;
+
 	/**
 	 * Check if the runtime is able to create a model given some ModelData.
 	 *
 	 * @param ModelData The model data for which to create a model.
 	 * @return True if the runtime is able to create the model, false otherwise.
 	 */
-	virtual bool CanCreateModelGPU(const TObjectPtr<UNNEModelData> ModelData) const = 0;
+	virtual ECanCreateModelGPUStatus CanCreateModelGPU(const TObjectPtr<UNNEModelData> ModelData) const = 0;
 	
 	/**
 	 * Create a model given some ModelData.
diff --git a/Engine/Source/Runtime/NNE/Public/NNERuntimeRDG.h b/Engine/Source/Runtime/NNE/Public/NNERuntimeRDG.h
--- a/Engine/Source/Runtime/NNE/Public/NNERuntimeRDG.h
+++ b/Engine/Source/Runtime/NNE/Public/NNERuntimeRDG.h
@@ -142,13 +142,15 @@ class NNE_API INNERuntimeRDG
 
 public:
 
+	using ECanCreateModelRDGStatus = UE::NNE::EResultStatus;
+
 	/**
 	 * Check if the runtime is able to create a model given some ModelData.
 	 *
 	 * @param ModelData The model data for which to create a model.
 	 * @return True if the runtime is able to create the model, false otherwise.
 	 */
-	virtual bool CanCreateModelRDG(const TObjectPtr<UNNEModelData> ModelData) const = 0;
+	virtual ECanCreateModelRDGStatus CanCreateModelRDG(const TObjectPtr<UNNEModelData> ModelData) const = 0;
 	
 	/**
 	 * Create a model given some ModelData.
