[NNE] Logs cleanup (v2)

Problem: logs were inconsistent across the different runtimes, witch could be confusing for final users.

Solution: make log consistent across NNE runtimes.

#tests win64 editor: smoke test, win64 standalone: IREE + style transfer (model and effects)

#rnx

diff --git a/Engine/Plugins/Experimental/NNERuntimeBasicCpu/Source/NNERuntimeBasicCpu/Private/NNERuntimeBasicCpu.cpp b/Engine/Plugins/Experimental/NNERuntimeBasicCpu/Source/NNERuntimeBasicCpu/Private/NNERuntimeBasicCpu.cpp
--- a/Engine/Plugins/Experimental/NNERuntimeBasicCpu/Source/NNERuntimeBasicCpu/Private/NNERuntimeBasicCpu.cpp
+++ b/Engine/Plugins/Experimental/NNERuntimeBasicCpu/Source/NNERuntimeBasicCpu/Private/NNERuntimeBasicCpu.cpp
@@ -1,6 +1,8 @@
 // Copyright Epic Games, Inc. All Rights Reserved.
 
 #include "NNERuntimeBasicCpu.h"
+
+#include "NNE.h"
 #include "NNERuntimeBasicCpuModel.h"
 #include "NNEModelData.h"
 
@@ -43,6 +45,7 @@ TSharedPtr<UE::NNE::FSharedModelData> UNNERuntimeBasicCpuImpl::CreateModelData(c
 {
 	if (CanCreateModelData(FileType, FileData, AdditionalFileData, FileId, TargetPlatform) != ECanCreateModelDataStatus::Ok)
 	{
+		UE_LOG(LogNNE, Warning, TEXT("UNNERuntimeBasicCpu cannot create the model data with id %s (Filetype: %s)"), *FileId.ToString(EGuidFormats::Digits).ToLower(), *FileType);
 		return nullptr;
 	}
 
@@ -103,6 +106,7 @@ TSharedPtr<UE::NNE::IModelCPU> UNNERuntimeBasicCpuImpl::CreateModelCPU(const TOb
 
 	if (CanCreateModelCPU(ModelData) != ECanCreateModelCPUStatus::Ok)
 	{
+		UE_LOG(LogNNE, Warning, TEXT("UNNERuntimeBasicCpu cannot create a model from the model data with id %s"), *ModelData->GetFileId().ToString(EGuidFormats::Digits));
 		return nullptr;
 	}
 
diff --git a/Engine/Plugins/Experimental/NNERuntimeIREE/Source/NNERuntimeIREE/Private/NNERuntimeIREE.cpp b/Engine/Plugins/Experimental/NNERuntimeIREE/Source/NNERuntimeIREE/Private/NNERuntimeIREE.cpp
--- a/Engine/Plugins/Experimental/NNERuntimeIREE/Source/NNERuntimeIREE/Private/NNERuntimeIREE.cpp
+++ b/Engine/Plugins/Experimental/NNERuntimeIREE/Source/NNERuntimeIREE/Private/NNERuntimeIREE.cpp
@@ -21,6 +21,7 @@
 #include "Misc/App.h"
 #include "Misc/Paths.h"
 #include "Misc/SecureHash.h"
+#include "NNE.h"
 #include "NNERuntimeIREECompiler.h"
 #include "NNERuntimeIREEModel.h"
 #include "NNERuntimeIREEMetaData.h"
@@ -122,7 +123,7 @@ TSharedPtr<UE::NNE::FSharedModelData> UNNERuntimeIREECpu::CreateModelData(const
 	FString TargetPlatformName = TargetPlatform ? TargetPlatform->IniPlatformName() : UGameplayStatics::GetPlatformName();
 	if (CanCreateModelData(FileType, FileData, AdditionalFileData, FileId, TargetPlatform) != ECanCreateModelDataStatus::Ok)
 	{
-		UE_LOG(LogTemp, Warning, TEXT("UNNERuntimeIREECpu cannot create the model data with id %s (Filetype: %s) for platform %s"), *FileId.ToString(EGuidFormats::Digits).ToLower(), *FileType, *TargetPlatformName);
+		UE_LOG(LogNNE, Warning, TEXT("UNNERuntimeIREECpu cannot create the model data with id %s (Filetype: %s) for platform %s"), *FileId.ToString(EGuidFormats::Digits).ToLower(), *FileType, *TargetPlatformName);
 		return TSharedPtr<UE::NNE::FSharedModelData>();
 	}
 
@@ -131,13 +132,13 @@ TSharedPtr<UE::NNE::FSharedModelData> UNNERuntimeIREECpu::CreateModelData(const
 	GetUpdatedPlatformConfig(TargetPlatformName, ConfigFile, ConfigFilePath);
 	if (ConfigFile.Dirty)
 	{
-		UE_LOG(LogTemp, Warning, TEXT("UNNERuntimeIREECpu could not find the required settings in config file %s. Please make the file writeable and re-start the editor or manually add the required staging settings or models will not work in packaged builds for platform %s!"), *ConfigFilePath, *TargetPlatformName);
+		UE_LOG(LogNNE, Warning, TEXT("UNNERuntimeIREECpu could not find the required settings in config file %s. Please make the file writeable and re-start the editor or manually add the required staging settings or models will not work in packaged builds for platform %s!"), *ConfigFilePath, *TargetPlatformName);
 	}
 
 	TUniquePtr<UE::NNERuntimeIREE::CPU::FCompiler> Compiler = UE::NNERuntimeIREE::CPU::FCompiler::Make(TargetPlatformName);
 	if (!Compiler.IsValid())
 	{
-		UE_LOG(LogTemp, Warning, TEXT("UNNERuntimeIREECpu failed to create a compiler to compile for platform %s"), *TargetPlatformName);
+		UE_LOG(LogNNE, Warning, TEXT("UNNERuntimeIREECpu failed to create a compiler to compile for platform %s"), *TargetPlatformName);
 		return TSharedPtr<UE::NNE::FSharedModelData>();
 	}
 	
@@ -152,7 +153,7 @@ TSharedPtr<UE::NNE::FSharedModelData> UNNERuntimeIREECpu::CreateModelData(const
 	FString StagingDir = FPaths::ConvertRelativePathToFull(FPaths::Combine(FPaths::ProjectDir(), GetStagedModelDirPath(TargetPlatformName)));
 	if (!Compiler->CompileMlir(FileData, FileIdString, IntermediateDir, StagingDir, CompilerResults, CompilerModuleMetaData))
 	{
-		UE_LOG(LogTemp, Warning, TEXT("UNNERuntimeIREECpu failed to compile model %s"), *FileIdString);
+		UE_LOG(LogNNE, Warning, TEXT("UNNERuntimeIREECpu failed to compile model %s"), *FileIdString);
 		return TSharedPtr<UE::NNE::FSharedModelData>();
 	}
 
@@ -162,7 +163,7 @@ TSharedPtr<UE::NNE::FSharedModelData> UNNERuntimeIREECpu::CreateModelData(const
 		FString StagedSharedLibPath = FPaths::Combine(StagingDir, CompilerResults[i].RelativeDirPath, CompilerResults[i].SharedLibraryFileName);
 		if (!FFileHelper::LoadFileToArray(SharedLibData, *StagedSharedLibPath) || SharedLibData.IsEmpty())
 		{
-			UE_LOG(LogTemp, Warning, TEXT("UNNERuntimeIREECpu could not read the shared library \"%s\""), *StagedSharedLibPath);
+			UE_LOG(LogNNE, Warning, TEXT("UNNERuntimeIREECpu could not read the shared library \"%s\""), *StagedSharedLibPath);
 			return TSharedPtr<UE::NNE::FSharedModelData>();
 		}
 		FSharedBuffer SharedLibBuffer = MakeSharedBufferFromArray(MoveTemp(SharedLibData));
@@ -172,7 +173,7 @@ TSharedPtr<UE::NNE::FSharedModelData> UNNERuntimeIREECpu::CreateModelData(const
 		FString StagedVmfbPath = FPaths::Combine(StagingDir, CompilerResults[i].RelativeDirPath, CompilerResults[i].VmfbFileName);
 		if (!FFileHelper::LoadFileToArray(VmfbData, *StagedVmfbPath) || VmfbData.IsEmpty())
 		{
-			UE_LOG(LogTemp, Warning, TEXT("UNNERuntimeIREECpu could not read the vmfb data \"%s\""), *StagedVmfbPath);
+			UE_LOG(LogNNE, Warning, TEXT("UNNERuntimeIREECpu could not read the vmfb data \"%s\""), *StagedVmfbPath);
 			return TSharedPtr<UE::NNE::FSharedModelData>();
 		}
 		FSharedBuffer VmfbBuffer = MakeSharedBufferFromArray(MoveTemp(VmfbData));
@@ -254,7 +255,7 @@ TSharedPtr<UE::NNE::IModelCPU> UNNERuntimeIREECpu::CreateModelCPU(const TObjectP
 
 	if (CanCreateModelCPU(ModelData) != ECanCreateModelCPUStatus::Ok)
 	{
-		UE_LOG(LogTemp, Warning, TEXT("UNNERuntimeIREECpu cannot create a model from the model data with id %s"), *ModelData->GetFileId().ToString(EGuidFormats::Digits));
+		UE_LOG(LogNNE, Warning, TEXT("UNNERuntimeIREECpu cannot create a model from the model data with id %s"), *ModelData->GetFileId().ToString(EGuidFormats::Digits));
 		return TSharedPtr<UE::NNE::IModelCPU>();
 	}
 
@@ -280,7 +281,7 @@ TSharedPtr<UE::NNE::IModelCPU> UNNERuntimeIREECpu::CreateModelCPU(const TObjectP
 	Reader << ModuleDataArray;
 	if (ModuleDataArray.IsEmpty())
 	{
-		UE_LOG(LogTemp, Warning, TEXT("UNNERuntimeIREECpu failed to find any module meta data, please reimport the original model"));
+		UE_LOG(LogNNE, Warning, TEXT("UNNERuntimeIREECpu failed to find any module meta data, please reimport the original model"));
 		return TSharedPtr<UE::NNE::IModelCPU>();
 	}
 	UNNERuntimeIREEModuleMetaData* ModuleMetaData = NewObject<UNNERuntimeIREEModuleMetaData>();
@@ -288,7 +289,7 @@ TSharedPtr<UE::NNE::IModelCPU> UNNERuntimeIREECpu::CreateModelCPU(const TObjectP
 	ModuleMetaData->Serialize(ObjectReader);
 	if (ModuleMetaData->FunctionMetaData.IsEmpty())
 	{
-		UE_LOG(LogTemp, Warning, TEXT("UNNERuntimeIREECpu failed to parse the module meta data, please reimport the original model"));
+		UE_LOG(LogNNE, Warning, TEXT("UNNERuntimeIREECpu failed to parse the module meta data, please reimport the original model"));
 		return TSharedPtr<UE::NNE::IModelCPU>();
 	}
 
@@ -335,7 +336,7 @@ TSharedPtr<UE::NNE::IModelCPU> UNNERuntimeIREECpu::CreateModelCPU(const TObjectP
 	}
 	if (!bFound)
 	{
-		UE_LOG(LogTemp, Warning, TEXT("UNNERuntimeIREECpu failed to find a matching architecture for \'%s\'"), *CurrentArchitecture);
+		UE_LOG(LogNNE, Warning, TEXT("UNNERuntimeIREECpu failed to find a matching architecture for \'%s\'"), *CurrentArchitecture);
 		return TSharedPtr<UE::NNE::IModelCPU>();
 	}
 
@@ -354,7 +355,7 @@ TSharedPtr<UE::NNE::IModelCPU> UNNERuntimeIREECpu::CreateModelCPU(const TObjectP
 		FSharedBuffer SharedBuffer = GetFromDDC(UE::NNERuntimeIREE::CPU::Private::GetModelDataIdentifier(GetRuntimeName(), UNNERuntimeIREECpu::GUID, FileIdString, UGameplayStatics::GetPlatformName(), Architecture) + "-lib");
 		if (SharedBuffer.GetSize() <= 0)
 		{
-			UE_LOG(LogTemp, Warning, TEXT("UNNERuntimeIREECpu could not fetch the shared library %s from DDC"), *SharedLibraryFileName);
+			UE_LOG(LogNNE, Warning, TEXT("UNNERuntimeIREECpu could not fetch the shared library %s from DDC"), *SharedLibraryFileName);
 			return TSharedPtr<UE::NNE::IModelCPU>();
 		}
 		PlatformFile.CreateDirectoryTree(*SharedLibraryDirPath);
@@ -367,7 +368,7 @@ TSharedPtr<UE::NNE::IModelCPU> UNNERuntimeIREECpu::CreateModelCPU(const TObjectP
 		FSharedBuffer SharedBuffer = GetFromDDC(UE::NNERuntimeIREE::CPU::Private::GetModelDataIdentifier(GetRuntimeName(), UNNERuntimeIREECpu::GUID, FileIdString, UGameplayStatics::GetPlatformName(), Architecture) + "-vmfb");
 		if (SharedBuffer.GetSize() <= 0)
 		{
-			UE_LOG(LogTemp, Warning, TEXT("UNNERuntimeIREECpu could not fetch the vmfb %s from DDC"), *VmfbFileName);
+			UE_LOG(LogNNE, Warning, TEXT("UNNERuntimeIREECpu could not fetch the vmfb %s from DDC"), *VmfbFileName);
 			return TSharedPtr<UE::NNE::IModelCPU>();
 		}
 		PlatformFile.CreateDirectoryTree(*SharedLibraryDirPath);
@@ -378,7 +379,7 @@ TSharedPtr<UE::NNE::IModelCPU> UNNERuntimeIREECpu::CreateModelCPU(const TObjectP
 	TSharedPtr<UE::NNE::IModelCPU> Model = UE::NNERuntimeIREE::CPU::FModel::Make(SharedLibraryDirPath, SharedLibraryFileName, VmfbFileName, SharedLibraryEntryPointName, *ModuleMetaData);
 	if (!Model.IsValid())
 	{
-		UE_LOG(LogTemp, Warning, TEXT("UNNERuntimeIREECpu could not initialize the model created from model data with id %s"), *FileIdString);
+		UE_LOG(LogNNE, Warning, TEXT("UNNERuntimeIREECpu could not initialize the model created from model data with id %s"), *FileIdString);
 		return TSharedPtr<UE::NNE::IModelCPU>();
 	}
 
@@ -467,6 +468,7 @@ TSharedPtr<UE::NNE::IModelGPU> UNNERuntimeIREEGpu::CreateModelGPU(const TObjectP
 
 	if (CanCreateModelGPU(ModelData) != ECanCreateModelGPUStatus::Ok)
 	{
+		UE_LOG(LogNNE, Warning, TEXT("UNNERuntimeIREEGpu cannot create a model from the model data with id %s"), *ModelData->GetFileId().ToString(EGuidFormats::Digits));
 		return TSharedPtr<UE::NNE::IModelGPU>();
 	}
 
@@ -607,6 +609,7 @@ TSharedPtr<UE::NNE::IModelRDG> UNNERuntimeIREERdg::CreateModelRDG(const TObjectP
 
 	if (CanCreateModelRDG(ModelData) != ECanCreateModelRDGStatus::Ok)
 	{
+		UE_LOG(LogNNE, Warning, TEXT("UNNERuntimeIREERdg cannot create a model from the model data with id %s"), *ModelData->GetFileId().ToString(EGuidFormats::Digits));
 		return TSharedPtr<UE::NNE::IModelRDG>();
 	}
 
diff --git a/Engine/Plugins/Experimental/NNERuntimeIREE/Source/NNERuntimeIREE/Private/NNERuntimeIREECompiler.cpp b/Engine/Plugins/Experimental/NNERuntimeIREE/Source/NNERuntimeIREE/Private/NNERuntimeIREECompiler.cpp
--- a/Engine/Plugins/Experimental/NNERuntimeIREE/Source/NNERuntimeIREE/Private/NNERuntimeIREECompiler.cpp
+++ b/Engine/Plugins/Experimental/NNERuntimeIREE/Source/NNERuntimeIREE/Private/NNERuntimeIREECompiler.cpp
@@ -14,6 +14,7 @@
 #include "Kismet/GameplayStatics.h"
 #include "Misc/FileHelper.h"
 #include "Misc/Paths.h"
+#include "NNE.h"
 #include "Serialization/JsonReader.h"
 #include "Serialization/JsonSerializer.h"
 #include "Templates/SharedPointer.h"
@@ -124,7 +125,7 @@ namespace UE::NNERuntimeIREE
 						{
 							if (BuildConfig.BuildTargets.IsEmpty())
 							{
-								UE_LOG(LogTemp, Warning, TEXT("UNNERuntimeIREECpu could not find targets in %s"), *BuildConfigFilePath);
+								UE_LOG(LogNNE, Warning, TEXT("UNNERuntimeIREECpu could not find targets in %s"), *BuildConfigFilePath);
 								continue;
 							}
 
@@ -143,12 +144,12 @@ namespace UE::NNERuntimeIREE
 								}
 								else
 								{
-									UE_LOG(LogTemp, Warning, TEXT("UNNERuntimeIREECpu could not replace environment variables in %s"), *BuildConfig.CompilerCommand[i]);
+									UE_LOG(LogNNE, Warning, TEXT("UNNERuntimeIREECpu could not replace environment variables in %s"), *BuildConfig.CompilerCommand[i]);
 								}
 							}
 							if (TmpCompilerCommand.IsEmpty())
 							{
-								UE_LOG(LogTemp, Warning, TEXT("UNNERuntimeIREECpu could not find the compiler executable in %s"), *BuildConfigFilePath);
+								UE_LOG(LogNNE, Warning, TEXT("UNNERuntimeIREECpu could not find the compiler executable in %s"), *BuildConfigFilePath);
 								continue;
 							}
 
@@ -167,12 +168,12 @@ namespace UE::NNERuntimeIREE
 								}
 								else
 								{
-									UE_LOG(LogTemp, Warning, TEXT("UNNERuntimeIREECpu could not replace environment variables in %s"), *BuildConfig.LinkerCommand[i]);
+									UE_LOG(LogNNE, Warning, TEXT("UNNERuntimeIREECpu could not replace environment variables in %s"), *BuildConfig.LinkerCommand[i]);
 								}
 							}
 							if (TmpLinkerCommand.IsEmpty())
 							{
-								UE_LOG(LogTemp, Warning, TEXT("UNNERuntimeIREECpu could not find the linker executable in %s"), *BuildConfigFilePath);
+								UE_LOG(LogNNE, Warning, TEXT("UNNERuntimeIREECpu could not find the linker executable in %s"), *BuildConfigFilePath);
 								continue;
 							}
 
@@ -184,12 +185,12 @@ namespace UE::NNERuntimeIREE
 						}
 						else
 						{
-							UE_LOG(LogTemp, Warning, TEXT("UNNERuntimeIREECpu could not parse build config file %s"), *BuildConfigFilePath);
+							UE_LOG(LogNNE, Warning, TEXT("UNNERuntimeIREECpu could not parse build config file %s"), *BuildConfigFilePath);
 						}
 					}
 					else
 					{
-						UE_LOG(LogTemp, Warning, TEXT("UNNERuntimeIREECpu could not read build config file %s"), *BuildConfigFilePath);
+						UE_LOG(LogNNE, Warning, TEXT("UNNERuntimeIREECpu could not read build config file %s"), *BuildConfigFilePath);
 					}
 				}
 			}
@@ -229,7 +230,7 @@ namespace UE::NNERuntimeIREE
 				FString CompilerArguments = BuildTargets[i].CompilerArguments;
 				if (!ResolveEnvironmentVariables(CompilerArguments))
 				{
-					UE_LOG(LogTemp, Warning, TEXT("UNNERuntimeIREECpu could not replace environment variables in %s"), *BuildTargets[i].CompilerArguments);
+					UE_LOG(LogNNE, Warning, TEXT("UNNERuntimeIREECpu could not replace environment variables in %s"), *BuildTargets[i].CompilerArguments);
 					continue;
 				}
 				CompilerArguments.ReplaceInline(*FString("${OBJECT_PATH}"), *(FString("\"") + ObjectFilePath + "\""));
@@ -240,8 +241,8 @@ namespace UE::NNERuntimeIREE
 
 				if (!PlatformFile.FileExists(*ObjectFilePath) || !PlatformFile.FileExists(*VmfbFilePath))
 				{
-					UE_LOG(LogTemp, Warning, TEXT("UNNERuntimeIREECpu failed to compile the model \"%s\" using the command:"), *InputFilePath);
-					UE_LOG(LogTemp, Warning, TEXT("\"%s\" %s"), *CompilerCommand, *CompilerArguments);
+					UE_LOG(LogNNE, Warning, TEXT("UNNERuntimeIREECpu failed to compile the model \"%s\" using the command:"), *InputFilePath);
+					UE_LOG(LogNNE, Warning, TEXT("\"%s\" %s"), *CompilerCommand, *CompilerArguments);
 					bResult = false;
 					continue;
 				}
@@ -249,7 +250,7 @@ namespace UE::NNERuntimeIREE
 				FString LinkerArguments = BuildTargets[i].LinkerArguments;
 				if (!ResolveEnvironmentVariables(LinkerArguments))
 				{
-					UE_LOG(LogTemp, Warning, TEXT("UNNERuntimeIREECpu could not replace environment variables in %s"), *BuildTargets[i].LinkerArguments);
+					UE_LOG(LogNNE, Warning, TEXT("UNNERuntimeIREECpu could not replace environment variables in %s"), *BuildTargets[i].LinkerArguments);
 					bResult = false;
 					continue;
 				}
@@ -260,8 +261,8 @@ namespace UE::NNERuntimeIREE
 
 				if (!PlatformFile.FileExists(*SharedLibFilePath))
 				{
-					UE_LOG(LogTemp, Warning, TEXT("UNNERuntimeIREECpu failed to link the model \"%s\" using the command:"), *InputFilePath);
-					UE_LOG(LogTemp, Warning, TEXT("\"%s\" %s"), *LinkerCommand, *LinkerArguments);
+					UE_LOG(LogNNE, Warning, TEXT("UNNERuntimeIREECpu failed to link the model \"%s\" using the command:"), *InputFilePath);
+					UE_LOG(LogNNE, Warning, TEXT("\"%s\" %s"), *LinkerCommand, *LinkerArguments);
 					bResult = false;
 					continue;
 				}
@@ -271,12 +272,12 @@ namespace UE::NNERuntimeIREE
 				FString StagedSharedLibFilePath = StagingFilePathNoExt + SharedLibExt;
 				if (IFileManager::Get().Copy(*StagedSharedLibFilePath, *SharedLibFilePath) != COPY_OK)
 				{
-					UE_LOG(LogTemp, Warning, TEXT("UNNERuntimeIREECpu failed to copy \"%s\" to \"%s\""), *SharedLibFilePath, *StagedSharedLibFilePath);
+					UE_LOG(LogNNE, Warning, TEXT("UNNERuntimeIREECpu failed to copy \"%s\" to \"%s\""), *SharedLibFilePath, *StagedSharedLibFilePath);
 					bResult = false;
 				}
 				if (IFileManager::Get().Copy(*StagedVmfbFilePath, *VmfbFilePath) != COPY_OK)
 				{
-					UE_LOG(LogTemp, Warning, TEXT("UNNERuntimeIREECpu failed to copy \"%s\" to \"%s\""), *VmfbFilePath, *StagedVmfbFilePath);
+					UE_LOG(LogNNE, Warning, TEXT("UNNERuntimeIREECpu failed to copy \"%s\" to \"%s\""), *VmfbFilePath, *StagedVmfbFilePath);
 					bResult = false;
 				}
 
@@ -284,21 +285,21 @@ namespace UE::NNERuntimeIREE
 				FString HeaderPath = IntermediateFilePathNoExt + ".h";
 				if (!PlatformFile.FileExists(*HeaderPath))
 				{
-					UE_LOG(LogTemp, Warning, TEXT("UNNERuntimeIREECpu could not find the model header \"%s\""), *HeaderPath);
+					UE_LOG(LogNNE, Warning, TEXT("UNNERuntimeIREECpu could not find the model header \"%s\""), *HeaderPath);
 					bResult = false;
 					continue;
 				}
 				FString HeaderString;
 				if (!FFileHelper::LoadFileToString(HeaderString, *HeaderPath) || HeaderString.IsEmpty())
 				{
-					UE_LOG(LogTemp, Warning, TEXT("UNNERuntimeIREECpu could not read the model header \"%s\""), *HeaderPath);
+					UE_LOG(LogNNE, Warning, TEXT("UNNERuntimeIREECpu could not read the model header \"%s\""), *HeaderPath);
 					bResult = false;
 					continue;
 				}
 				SharedLibraryEntryPointName = GetSharedLibraryEntryPointName(HeaderString);
 				if (SharedLibraryEntryPointName.IsEmpty())
 				{
-					UE_LOG(LogTemp, Warning, TEXT("UNNERuntimeIREECpu could not find the entry point in model header \"%s\""), *HeaderPath);
+					UE_LOG(LogNNE, Warning, TEXT("UNNERuntimeIREECpu could not find the entry point in model header \"%s\""), *HeaderPath);
 					bResult = false;
 					continue;
 				}
diff --git a/Engine/Plugins/Experimental/NNERuntimeIREE/Source/NNERuntimeIREE/Private/NNERuntimeIREEMetaData.cpp b/Engine/Plugins/Experimental/NNERuntimeIREE/Source/NNERuntimeIREE/Private/NNERuntimeIREEMetaData.cpp
--- a/Engine/Plugins/Experimental/NNERuntimeIREE/Source/NNERuntimeIREE/Private/NNERuntimeIREEMetaData.cpp
+++ b/Engine/Plugins/Experimental/NNERuntimeIREE/Source/NNERuntimeIREE/Private/NNERuntimeIREEMetaData.cpp
@@ -4,6 +4,7 @@
 
 #include "Containers/Array.h"
 #include "Serialization/CustomVersion.h"
+#include "NNE.h"
 
 namespace UE::NNERuntimeIREE::ModuleMetaData::Private
 {
@@ -325,7 +326,7 @@ void UNNERuntimeIREEModuleMetaData::Serialize(FArchive& Ar)
 			}
 			break;
 		default:
-			UE_LOG(LogTemp, Error, TEXT("UNNERuntimeIREEModuleMetaData: Unknown asset version %d: Deserialisation failed, please reimport the original model."), Ar.CustomVer(UE::NNERuntimeIREE::ModuleMetaData::Private::GUID));
+			UE_LOG(LogNNE, Error, TEXT("UNNERuntimeIREEModuleMetaData: Unknown asset version %d: Deserialisation failed, please reimport the original model."), Ar.CustomVer(UE::NNERuntimeIREE::ModuleMetaData::Private::GUID));
 			break;
 		}
 	}
diff --git a/Engine/Plugins/Experimental/NNERuntimeIREE/Source/NNERuntimeIREE/Private/NNERuntimeIREEModel.cpp b/Engine/Plugins/Experimental/NNERuntimeIREE/Source/NNERuntimeIREE/Private/NNERuntimeIREEModel.cpp
--- a/Engine/Plugins/Experimental/NNERuntimeIREE/Source/NNERuntimeIREE/Private/NNERuntimeIREEModel.cpp
+++ b/Engine/Plugins/Experimental/NNERuntimeIREE/Source/NNERuntimeIREE/Private/NNERuntimeIREEModel.cpp
@@ -9,6 +9,7 @@
 #include "Memory/SharedBuffer.h"
 #include "Misc/Paths.h"
 #include "Misc/ScopeLock.h"
+#include "NNE.h"
 #include "NNEModelData.h"
 #include "NNEStatus.h"
 #include "Serialization/Archive.h"
@@ -102,7 +103,7 @@ namespace UE::NNERuntimeIREE
 			void* ErrorString = FMemory::Malloc(TrueLength + 1);
 			((char*)ErrorString)[TrueLength] = (char)0;
 			iree_status_format(InStatus, TrueLength, (char*)ErrorString, &TrueLength);
-			UE_LOG(LogTemp, Error, TEXT("%s: %s"), *InMessage, *FString(StringCast<TCHAR>(static_cast<const ANSICHAR*>(ErrorString)).Get()));
+			UE_LOG(LogNNE, Error, TEXT("%s: %s"), *InMessage, *FString(StringCast<TCHAR>(static_cast<const ANSICHAR*>(ErrorString)).Get()));
 			FMemory::Free(ErrorString);
 		}
 
@@ -352,13 +353,13 @@ namespace UE::NNERuntimeIREE
 				TUniquePtr<FArchive> Reader = TUniquePtr<FArchive>(IFileManager::Get().CreateFileReader(*CombinedPath, 0));
 				if (!Reader)
 				{
-					UE_LOG(LogTemp, Error, TEXT("UE::NNERuntimeIREE::Private::FModule failed to open the vmfb data file '%s'"), *CombinedPath);
+					UE_LOG(LogNNE, Error, TEXT("UE::NNERuntimeIREE::Private::FModule failed to open the vmfb data file '%s'"), *CombinedPath);
 					return TSharedPtr<FModule>();
 				}
 				int64 DataSize = Reader->TotalSize();
 				if (DataSize < 1)
 				{
-					UE_LOG(LogTemp, Error, TEXT("UE::NNERuntimeIREE::Private::FModule's vmfb data file '%s' is empty"), *CombinedPath);
+					UE_LOG(LogNNE, Error, TEXT("UE::NNERuntimeIREE::Private::FModule's vmfb data file '%s' is empty"), *CombinedPath);
 					return TSharedPtr<FModule>();
 				}
 
@@ -437,7 +438,7 @@ namespace UE::NNERuntimeIREE
 
 				if (!bFound)
 				{
-					UE_LOG(LogTemp, Error, TEXT("UE::NNERuntimeIREE::Private::FModule failed to find the module function %s"), *InFunctionName);
+					UE_LOG(LogNNE, Error, TEXT("UE::NNERuntimeIREE::Private::FModule failed to find the module function %s"), *InFunctionName);
 					iree_status_free(Status);
 					return false;
 				}
@@ -485,7 +486,7 @@ namespace UE::NNERuntimeIREE
 					void* Library = FPlatformProcess::GetDllHandle(*CombinedPath);
 					if (!Library)
 					{
-						UE_LOG(LogTemp, Error, TEXT("UE::NNERuntimeIREE::CPU::Private::FLibrary failed to load the shared library '%s'"), *CombinedPath);
+						UE_LOG(LogNNE, Error, TEXT("UE::NNERuntimeIREE::CPU::Private::FLibrary failed to load the shared library '%s'"), *CombinedPath);
 						return TSharedPtr<FLibrary>();
 					}
 #else
@@ -494,7 +495,7 @@ namespace UE::NNERuntimeIREE
 					FPlatformProcess::PopDllDirectory(*InLibraryPath);
 					if (!Library)
 					{
-						UE_LOG(LogTemp, Error, TEXT("UE::NNERuntimeIREE::CPU::Private::FLibrary failed to load the shared library '%s' from '%s'"), *InLibraryName, *InLibraryPath);
+						UE_LOG(LogNNE, Error, TEXT("UE::NNERuntimeIREE::CPU::Private::FLibrary failed to load the shared library '%s' from '%s'"), *InLibraryName, *InLibraryPath);
 						return TSharedPtr<FLibrary>();
 					}
 #endif
@@ -512,7 +513,7 @@ namespace UE::NNERuntimeIREE
 						*OutFunctionPointer = Result;
 						return true;
 					}
-					UE_LOG(LogTemp, Error, TEXT("UE::NNERuntimeIREE::CPU::Private::FLibrary failed to get the function %s"), *InFunctionName);
+					UE_LOG(LogNNE, Error, TEXT("UE::NNERuntimeIREE::CPU::Private::FLibrary failed to get the function %s"), *InFunctionName);
 					return false;
 				}
 			};
@@ -682,7 +683,7 @@ namespace UE::NNERuntimeIREE
 					TConstArrayView<UE::NNE::FTensorDesc> OutputTensorDescs = InModule->GetFunctionMetaDataView()[0].OutputDescs;
 					if (!iree_status_is_ok(Status) || NumInputs != InputTensorDescs.Num() || NumOutputs != OutputTensorDescs.Num())
 					{
-						UE_LOG(LogTemp, Error, TEXT("UE::NNERuntimeIREE::CPU::Private::FSession has a function signature mismatch in function %s"), *MainFunctionName);
+						UE_LOG(LogNNE, Error, TEXT("UE::NNERuntimeIREE::CPU::Private::FSession has a function signature mismatch in function %s"), *MainFunctionName);
 						iree_runtime_session_release(Session);
 						iree_status_free(Status);
 						return TSharedPtr<FSession>();
@@ -831,7 +832,7 @@ namespace UE::NNERuntimeIREE
 							iree_hal_buffer_t* Buffer = iree_hal_buffer_view_buffer(BufferViews[i]);
 							if (!Buffer)
 							{
-								UE_LOG(LogTemp, Error, TEXT("UE::NNERuntimeIREE::CPU::Private::FSession failed to get the result buffer"));
+								UE_LOG(LogNNE, Error, TEXT("UE::NNERuntimeIREE::CPU::Private::FSession failed to get the result buffer"));
 								Result = ERunSyncStatus::Fail;
 								break;
 							}
diff --git a/Engine/Plugins/Experimental/NNERuntimeIREE/Source/NNERuntimeIREEEditor/Private/NNERuntimeIREEModelDataFactory.cpp b/Engine/Plugins/Experimental/NNERuntimeIREE/Source/NNERuntimeIREEEditor/Private/NNERuntimeIREEModelDataFactory.cpp
--- a/Engine/Plugins/Experimental/NNERuntimeIREE/Source/NNERuntimeIREEEditor/Private/NNERuntimeIREEModelDataFactory.cpp
+++ b/Engine/Plugins/Experimental/NNERuntimeIREE/Source/NNERuntimeIREEEditor/Private/NNERuntimeIREEModelDataFactory.cpp
@@ -12,6 +12,7 @@
 #include "Interfaces/IMainFrameModule.h"
 #include "Kismet/GameplayStatics.h"
 #include "Modules/ModuleManager.h"
+#include "NNE.h"
 #include "NNEModelData.h"
 #include "NNERuntimeIREEMetaData.h"
 #include "Serialization/MemoryWriter.h"
@@ -163,7 +164,7 @@ UObject* UNNERuntimeIREEModelDataFactory::FactoryCreateBinary(UClass* Class, UOb
 	TObjectPtr<UNNERuntimeIREEModuleMetaData> ModuleMetaData = NewObject<UNNERuntimeIREEModuleMetaData>();
 	if (!ModuleMetaData->ParseFromString(FileDataString) || ModuleMetaData->FunctionMetaData.IsEmpty())
 	{
-		UE_LOG(LogTemp, Error, TEXT("UNNERuntimeIREEModelDataFactory failed to parse the models meta data"));
+		UE_LOG(LogNNE, Error, TEXT("UNNERuntimeIREEModelDataFactory failed to parse the models meta data"));
 		return nullptr;
 	}
 
@@ -193,7 +194,7 @@ UObject* UNNERuntimeIREEModelDataFactory::FactoryCreateBinary(UClass* Class, UOb
 
 		if (!ImportWindow->ImportButtonClicked())
 		{
-			UE_LOG(LogTemp, Error, TEXT("UNNERuntimeIREEModelDataFactory could not import the model! Please select the neural network main function in the import dialog!"));
+			UE_LOG(LogNNE, Error, TEXT("UNNERuntimeIREEModelDataFactory could not import the model! Please select the neural network main function in the import dialog!"));
 			return nullptr;
 		}
 
diff --git a/Engine/Plugins/Experimental/NNERuntimeRDG/Source/NNERuntimeRDG/Private/NNERuntimeRDGDml.cpp b/Engine/Plugins/Experimental/NNERuntimeRDG/Source/NNERuntimeRDG/Private/NNERuntimeRDGDml.cpp
--- a/Engine/Plugins/Experimental/NNERuntimeRDG/Source/NNERuntimeRDG/Private/NNERuntimeRDGDml.cpp
+++ b/Engine/Plugins/Experimental/NNERuntimeRDG/Source/NNERuntimeRDG/Private/NNERuntimeRDGDml.cpp
@@ -9,6 +9,7 @@
 #include "Interfaces/IPluginManager.h"
 #include "Kismet/GameplayStatics.h"
 #include "Misc/SecureHash.h"
+#include "NNE.h"
 #include "NNEModelData.h"
 #include "NNERuntimeFormat.h"
 #include "NNERuntimeRDGBase.h"
@@ -125,6 +126,7 @@ TSharedPtr<UE::NNE::FSharedModelData> UNNERuntimeRDGDmlImpl::CreateModelData(con
 {
 	if (CanCreateModelData(FileType, FileData, AdditionalFileData, FileId, TargetPlatform) != ECanCreateModelDataStatus::Ok)
 	{
+		UE_LOG(LogNNE, Warning, TEXT("UNNERuntimeRDGDml cannot create the model data with id %s (Filetype: %s)"), *FileId.ToString(EGuidFormats::Digits).ToLower(), *FileType);
 		return {};
 	}
 
@@ -200,6 +202,7 @@ TSharedPtr<UE::NNE::IModelRDG> UNNERuntimeRDGDmlImpl::CreateModelRDG(const TObje
 #ifdef NNE_USE_DIRECTML
 	if (CanCreateModelRDG(ModelData) != ECanCreateModelRDGStatus::Ok)
 	{
+		UE_LOG(LogNNE, Warning, TEXT("UNNERuntimeRDGDml cannot create a model from the model data with id %s"), *ModelData->GetFileId().ToString(EGuidFormats::Digits));
 		return TSharedPtr<UE::NNE::IModelRDG>();
 	}
 
diff --git a/Engine/Plugins/Experimental/NNERuntimeRDG/Source/NNERuntimeRDG/Private/NNERuntimeRDGHlsl.cpp b/Engine/Plugins/Experimental/NNERuntimeRDG/Source/NNERuntimeRDG/Private/NNERuntimeRDGHlsl.cpp
--- a/Engine/Plugins/Experimental/NNERuntimeRDG/Source/NNERuntimeRDG/Private/NNERuntimeRDGHlsl.cpp
+++ b/Engine/Plugins/Experimental/NNERuntimeRDG/Source/NNERuntimeRDG/Private/NNERuntimeRDGHlsl.cpp
@@ -5,6 +5,7 @@
 #include "EngineAnalytics.h"
 #include "Kismet/GameplayStatics.h"
 #include "Misc/SecureHash.h"
+#include "NNE.h"
 #include "NNEAttributeMap.h"
 #include "NNEModelData.h"
 #include "NNEModelOptimizerInterface.h"
@@ -122,6 +123,7 @@ TSharedPtr<UE::NNE::FSharedModelData> UNNERuntimeRDGHlslImpl::CreateModelData(co
 {
 	if (CanCreateModelData(FileType, FileData, AdditionalFileData, FileId, TargetPlatform) != ECanCreateModelDataStatus::Ok)
 	{
+		UE_LOG(LogNNE, Warning, TEXT("UNNERuntimeRDGHlsl cannot create the model data with id %s (Filetype: %s)"), *FileId.ToString(EGuidFormats::Digits).ToLower(), *FileType);
 		return {};
 	}
 
@@ -161,6 +163,7 @@ TSharedPtr<UE::NNE::IModelRDG> UNNERuntimeRDGHlslImpl::CreateModelRDG(const TObj
 {
 	if (CanCreateModelRDG(ModelData) != ECanCreateModelRDGStatus::Ok)
 	{
+		UE_LOG(LogNNE, Warning, TEXT("UNNERuntimeRDGHlsl cannot create a model from the model data with id %s"), *ModelData->GetFileId().ToString(EGuidFormats::Digits));
 		return TSharedPtr<UE::NNE::IModelRDG>();
 	}
 
diff --git a/Engine/Plugins/NNE/NNERuntimeORT/Source/NNERuntimeORT/Private/NNERuntimeORT.cpp b/Engine/Plugins/NNE/NNERuntimeORT/Source/NNERuntimeORT/Private/NNERuntimeORT.cpp
--- a/Engine/Plugins/NNE/NNERuntimeORT/Source/NNERuntimeORT/Private/NNERuntimeORT.cpp
+++ b/Engine/Plugins/NNE/NNERuntimeORT/Source/NNERuntimeORT/Private/NNERuntimeORT.cpp
@@ -5,6 +5,7 @@
 #include "EngineAnalytics.h"
 #include "Kismet/GameplayStatics.h"
 #include "Misc/SecureHash.h"
+#include "NNE.h"
 #include "NNEAttributeMap.h"
 #include "NNEModelData.h"
 #include "NNEModelOptimizerInterface.h"
@@ -38,6 +39,7 @@ TSharedPtr<UE::NNE::FSharedModelData> UNNERuntimeORTDml::CreateModelData(const F
 {
 	if (CanCreateModelData(FileType, FileData, AdditionalFileData, FileId, TargetPlatform) != ECanCreateModelDataStatus::Ok)
 	{
+		UE_LOG(LogNNE, Warning, TEXT("UNNERuntimeORTDml cannot create the model data with id %s (Filetype: %s)"), *FileId.ToString(EGuidFormats::Digits).ToLower(), *FileType);
 		return {};
 	}
 
@@ -66,6 +68,7 @@ TSharedPtr<UE::NNE::FSharedModelData> UNNERuntimeORTCpu::CreateModelData(const F
 {
 	if (CanCreateModelData(FileType, FileData, AdditionalFileData, FileId, TargetPlatform) != ECanCreateModelDataStatus::Ok)
 	{
+		UE_LOG(LogNNE, Warning, TEXT("UNNERuntimeORTCpu cannot create the model data with id %s (Filetype: %s)"), *FileId.ToString(EGuidFormats::Digits).ToLower(), *FileType);
 		return {};
 	}
 
@@ -155,6 +158,7 @@ TSharedPtr<UE::NNE::IModelCPU> UNNERuntimeORTCpu::CreateModelCPU(const TObjectPt
 
 	if (CanCreateModelCPU(ModelData) != ECanCreateModelCPUStatus::Ok)
 	{
+		UE_LOG(LogNNE, Warning, TEXT("UNNERuntimeORTCpu cannot create a model from the model data with id %s"), *ModelData->GetFileId().ToString(EGuidFormats::Digits));
 		return TSharedPtr<UE::NNE::IModelCPU>();
 	}
 
@@ -220,6 +224,7 @@ TSharedPtr<UE::NNE::IModelGPU> UNNERuntimeORTDml::CreateModelGPU(const TObjectPt
 
 	if (CanCreateModelGPU(ModelData) != ECanCreateModelGPUStatus::Ok)
 	{
+		UE_LOG(LogNNE, Warning, TEXT("UNNERuntimeORTDml cannot create a model from the model data with id %s"), *ModelData->GetFileId().ToString(EGuidFormats::Digits));
 		return TSharedPtr<UE::NNE::IModelGPU>();
 	}
 
diff --git a/Engine/Source/Runtime/NNE/Private/NNEModelData.cpp b/Engine/Source/Runtime/NNE/Private/NNEModelData.cpp
--- a/Engine/Source/Runtime/NNE/Private/NNEModelData.cpp
+++ b/Engine/Source/Runtime/NNE/Private/NNEModelData.cpp
@@ -140,6 +140,10 @@ namespace UE::NNE::ModelData
 			{
 				return NNERuntime->CreateModelData(FileType, FileData, AdditionalFileData, FileId, TargetPlatform);
 			}
+			else
+			{
+				UE_LOG(LogNNE, Warning, TEXT("Runtime %s cannot create the model data with id %s (Filetype: %s)"), *RuntimeName , *FileId.ToString(EGuidFormats::Digits).ToLower(), *FileType);
+			}
 		}
 		else
 		{
@@ -197,13 +201,6 @@ namespace UE::NNE
 
 } // UE::NNE
 
-void UNNEModelData::GetAssetRegistryTags(TArray<FAssetRegistryTag>& OutTags) const
-{
-	PRAGMA_DISABLE_DEPRECATION_WARNINGS;
-	Super::GetAssetRegistryTags(OutTags);
-	PRAGMA_ENABLE_DEPRECATION_WARNINGS;
-}
-
 void UNNEModelData::GetAssetRegistryTags(FAssetRegistryTagsContext Context) const
 {
 	Context.AddTag(FAssetRegistryTag("TargetRuntimes", UE::NNE::ModelData::GetRuntimesAsString(GetTargetRuntimes()), FAssetRegistryTag::TT_Alphabetical));
diff --git a/Engine/Source/Runtime/NNE/Public/NNEModelData.h b/Engine/Source/Runtime/NNE/Public/NNEModelData.h
--- a/Engine/Source/Runtime/NNE/Public/NNEModelData.h
+++ b/Engine/Source/Runtime/NNE/Public/NNEModelData.h
@@ -78,8 +78,6 @@ public:
 
 	// UObject interface
 	virtual void GetAssetRegistryTags(FAssetRegistryTagsContext Context) const override;
-	UE_DEPRECATED(5.4, "Implement the version that takes FAssetRegistryTagsContext instead.")
-	virtual void GetAssetRegistryTags(TArray<FAssetRegistryTag>& OutTags) const override;
 	virtual void Serialize(FArchive& Ar) override;
 	// End of UObject interface
 
