[NNE] Improve logging when a fileid is not supported by a given runtime.

Problem:
* all nne assets need to be flagged with the correct runtimes they should be cooked for or a warning might be triggered at cook time.

Solution:
* improve logging by differentiating unsupported fileid for cooked runtime vs other source of failures.

#tests win64 editor: smoke test, win64 standalone: all hlsl and ort tests.

#rnx

diff --git a/Engine/Plugins/Experimental/NNERuntimeBasicCpu/Source/NNERuntimeBasicCpu/Private/NNERuntimeBasicCpu.cpp b/Engine/Plugins/Experimental/NNERuntimeBasicCpu/Source/NNERuntimeBasicCpu/Private/NNERuntimeBasicCpu.cpp
--- a/Engine/Plugins/Experimental/NNERuntimeBasicCpu/Source/NNERuntimeBasicCpu/Private/NNERuntimeBasicCpu.cpp
+++ b/Engine/Plugins/Experimental/NNERuntimeBasicCpu/Source/NNERuntimeBasicCpu/Private/NNERuntimeBasicCpu.cpp
@@ -15,7 +15,7 @@ UNNERuntimeBasicCpuImpl::ECanCreateModelDataStatus UNNERuntimeBasicCpuImpl::CanC
 {
 	if (FileType.Compare("ubnne", ESearchCase::IgnoreCase) != 0)
 	{
-		return ECanCreateModelDataStatus::Fail;
+		return ECanCreateModelDataStatus::FailFileIdNotSupported;
 	}
 
 	// We require at least a magic number and version number
diff --git a/Engine/Plugins/Experimental/NNERuntimeIREE/Source/NNERuntimeIREE/Private/NNERuntimeIREE.cpp b/Engine/Plugins/Experimental/NNERuntimeIREE/Source/NNERuntimeIREE/Private/NNERuntimeIREE.cpp
--- a/Engine/Plugins/Experimental/NNERuntimeIREE/Source/NNERuntimeIREE/Private/NNERuntimeIREE.cpp
+++ b/Engine/Plugins/Experimental/NNERuntimeIREE/Source/NNERuntimeIREE/Private/NNERuntimeIREE.cpp
@@ -109,7 +109,7 @@ FString UNNERuntimeIREECpu::GetRuntimeName() const
 UNNERuntimeIREECpu::ECanCreateModelDataStatus UNNERuntimeIREECpu::CanCreateModelData(const FString& FileType, TConstArrayView<uint8> FileData, const TMap<FString, TConstArrayView<uint8>>& AdditionalFileData, const FGuid& FileId, const ITargetPlatform* TargetPlatform) const
 {
 #if WITH_EDITOR
-	return FileType.Compare(TEXT("mlir"), ESearchCase::IgnoreCase) == 0 ? ECanCreateModelDataStatus::Ok : ECanCreateModelDataStatus::Fail;
+	return FileType.Compare(TEXT("mlir"), ESearchCase::IgnoreCase) == 0 ? ECanCreateModelDataStatus::Ok : ECanCreateModelDataStatus::FailFileIdNotSupported;
 #else
 	return ECanCreateModelDataStatus::Fail;
 #endif // WITH_EDITOR
@@ -419,7 +419,7 @@ FString UNNERuntimeIREEGpu::GetRuntimeName() const
 UNNERuntimeIREEGpu::ECanCreateModelDataStatus UNNERuntimeIREEGpu::CanCreateModelData(const FString& FileType, TConstArrayView<uint8> FileData, const TMap<FString, TConstArrayView<uint8>>& AdditionalFileData, const FGuid& FileId, const ITargetPlatform* TargetPlatform) const
 {
 #if WITH_EDITOR
-	return FileType.Compare(TEXT("mlir"), ESearchCase::IgnoreCase) == 0 ? ECanCreateModelDataStatus::Ok : ECanCreateModelDataStatus::Fail;
+	return FileType.Compare(TEXT("mlir"), ESearchCase::IgnoreCase) == 0 ? ECanCreateModelDataStatus::Ok : ECanCreateModelDataStatus::FailFileIdNotSupported;
 #else
 	return ECanCreateModelDataStatus::Fail;
 #endif // WITH_EDITOR
@@ -562,7 +562,7 @@ FString UNNERuntimeIREERdg::GetRuntimeName() const
 UNNERuntimeIREERdg::ECanCreateModelDataStatus UNNERuntimeIREERdg::CanCreateModelData(const FString& FileType, TConstArrayView<uint8> FileData, const TMap<FString, TConstArrayView<uint8>>& AdditionalFileData, const FGuid& FileId, const ITargetPlatform* TargetPlatform) const
 {
 #if WITH_EDITOR
-	return	FileType.Compare(TEXT("mlir"), ESearchCase::IgnoreCase) == 0 ? ECanCreateModelDataStatus::Ok : ECanCreateModelDataStatus::Fail;
+	return	FileType.Compare(TEXT("mlir"), ESearchCase::IgnoreCase) == 0 ? ECanCreateModelDataStatus::Ok : ECanCreateModelDataStatus::FailFileIdNotSupported;
 #else
 	return ECanCreateModelDataStatus::Fail;
 #endif // WITH_EDITOR
diff --git a/Engine/Plugins/Experimental/NNERuntimeRDG/Source/NNERuntimeRDG/Private/NNERuntimeRDGDml.cpp b/Engine/Plugins/Experimental/NNERuntimeRDG/Source/NNERuntimeRDG/Private/NNERuntimeRDGDml.cpp
--- a/Engine/Plugins/Experimental/NNERuntimeRDG/Source/NNERuntimeRDG/Private/NNERuntimeRDGDml.cpp
+++ b/Engine/Plugins/Experimental/NNERuntimeRDG/Source/NNERuntimeRDG/Private/NNERuntimeRDGDml.cpp
@@ -115,7 +115,7 @@ using namespace UE::NNERuntimeRDG::Private::Dml;
 UNNERuntimeRDGDmlImpl::ECanCreateModelDataStatus UNNERuntimeRDGDmlImpl::CanCreateModelData(const FString& FileType, TConstArrayView<uint8> FileData, const TMap<FString, TConstArrayView<uint8>>& AdditionalFileData, const FGuid& FileId, const ITargetPlatform* TargetPlatform) const
 {
 #ifdef NNE_UTILITIES_AVAILABLE
-	return FileType.Compare("onnx", ESearchCase::IgnoreCase) == 0 ? ECanCreateModelDataStatus::Ok : ECanCreateModelDataStatus::Fail;
+	return FileType.Compare("onnx", ESearchCase::IgnoreCase) == 0 ? ECanCreateModelDataStatus::Ok : ECanCreateModelDataStatus::FailFileIdNotSupported;
 #else
 	UE_LOG(LogNNE, Display, TEXT("NNEUtilities is not available on this platform"));
 	return ECanCreateModelDataStatus::Fail;
diff --git a/Engine/Plugins/Experimental/NNERuntimeRDG/Source/NNERuntimeRDG/Private/NNERuntimeRDGHlsl.cpp b/Engine/Plugins/Experimental/NNERuntimeRDG/Source/NNERuntimeRDG/Private/NNERuntimeRDGHlsl.cpp
--- a/Engine/Plugins/Experimental/NNERuntimeRDG/Source/NNERuntimeRDG/Private/NNERuntimeRDGHlsl.cpp
+++ b/Engine/Plugins/Experimental/NNERuntimeRDG/Source/NNERuntimeRDG/Private/NNERuntimeRDGHlsl.cpp
@@ -89,7 +89,7 @@ bool UNNERuntimeRDGHlslImpl::Init()
 UNNERuntimeRDGHlslImpl::ECanCreateModelDataStatus UNNERuntimeRDGHlslImpl::CanCreateModelData(const FString& FileType, TConstArrayView<uint8> FileData, const TMap<FString, TConstArrayView<uint8>>& AdditionalFileData, const FGuid& FileId, const ITargetPlatform* TargetPlatform) const
 {
 #ifdef NNE_UTILITIES_AVAILABLE
-	return FileType.Compare("onnx", ESearchCase::IgnoreCase) == 0 ? ECanCreateModelDataStatus::Ok : ECanCreateModelDataStatus::Fail;
+	return FileType.Compare("onnx", ESearchCase::IgnoreCase) == 0 ? ECanCreateModelDataStatus::Ok : ECanCreateModelDataStatus::FailFileIdNotSupported;
 #else
 	UE_LOG(LogNNE, Display, TEXT("NNEUtilities is not available on this platform"));
 	return ECanCreateModelDataStatus::Fail;
diff --git a/Engine/Plugins/NNE/NNERuntimeORT/Source/NNERuntimeORT/Private/NNERuntimeORT.cpp b/Engine/Plugins/NNE/NNERuntimeORT/Source/NNERuntimeORT/Private/NNERuntimeORT.cpp
--- a/Engine/Plugins/NNE/NNERuntimeORT/Source/NNERuntimeORT/Private/NNERuntimeORT.cpp
+++ b/Engine/Plugins/NNE/NNERuntimeORT/Source/NNERuntimeORT/Private/NNERuntimeORT.cpp
@@ -27,12 +27,12 @@ int32 UNNERuntimeORTCpu::Version = 0x00000001;
 
 UNNERuntimeORTDml::ECanCreateModelDataStatus UNNERuntimeORTDml::CanCreateModelData(const FString& FileType, TConstArrayView<uint8> FileData, const TMap<FString, TConstArrayView<uint8>>& AdditionalFileData, const FGuid& FileId, const ITargetPlatform* TargetPlatform) const
 {
-	return (!FileData.IsEmpty() && FileType.Compare("onnx", ESearchCase::IgnoreCase) == 0) ? ECanCreateModelDataStatus::Ok : ECanCreateModelDataStatus::Fail;
+	return (!FileData.IsEmpty() && FileType.Compare("onnx", ESearchCase::IgnoreCase) == 0) ? ECanCreateModelDataStatus::Ok : ECanCreateModelDataStatus::FailFileIdNotSupported;
 }
 
 UNNERuntimeORTCpu::ECanCreateModelDataStatus UNNERuntimeORTCpu::CanCreateModelData(const FString& FileType, TConstArrayView<uint8> FileData, const TMap<FString, TConstArrayView<uint8>>& AdditionalFileData, const FGuid& FileId, const ITargetPlatform* TargetPlatform) const
 {
-	return (!FileData.IsEmpty() && FileType.Compare("onnx", ESearchCase::IgnoreCase) == 0) ? ECanCreateModelDataStatus::Ok : ECanCreateModelDataStatus::Fail;
+	return (!FileData.IsEmpty() && FileType.Compare("onnx", ESearchCase::IgnoreCase) == 0) ? ECanCreateModelDataStatus::Ok : ECanCreateModelDataStatus::FailFileIdNotSupported;
 }
 
 TSharedPtr<UE::NNE::FSharedModelData> UNNERuntimeORTDml::CreateModelData(const FString& FileType, TConstArrayView<uint8> FileData, const TMap<FString, TConstArrayView<uint8>>& AdditionalFileData, const FGuid& FileId, const ITargetPlatform* TargetPlatform)
diff --git a/Engine/Source/Runtime/NNE/Private/NNEModelData.cpp b/Engine/Source/Runtime/NNE/Private/NNEModelData.cpp
--- a/Engine/Source/Runtime/NNE/Private/NNEModelData.cpp
+++ b/Engine/Source/Runtime/NNE/Private/NNEModelData.cpp
@@ -136,13 +136,18 @@ namespace UE::NNE::ModelData
 		TWeakInterfacePtr<INNERuntime> NNERuntime = UE::NNE::GetRuntime<INNERuntime>(RuntimeName);
 		if (NNERuntime.IsValid())
 		{
-			if (NNERuntime->CanCreateModelData(FileType, FileData, AdditionalFileData, FileId, TargetPlatform) == INNERuntime::ECanCreateModelDataStatus::Ok)
+			INNERuntime::ECanCreateModelDataStatus CanCreateModelDataStatus = NNERuntime->CanCreateModelData(FileType, FileData, AdditionalFileData, FileId, TargetPlatform);
+			if (CanCreateModelDataStatus == INNERuntime::ECanCreateModelDataStatus::Ok)
 			{
 				return NNERuntime->CreateModelData(FileType, FileData, AdditionalFileData, FileId, TargetPlatform);
 			}
+			else if (CanCreateModelDataStatus == INNERuntime::ECanCreateModelDataStatus::FailFileIdNotSupported)
+			{
+				UE_LOG(LogNNE, Display, TEXT("Runtime %s does not support Filetype: %s, skipping the model data creation for model with id %s "), *RuntimeName, *FileType, *FileId.ToString(EGuidFormats::Digits).ToLower());
+			}
 			else
 			{
-				UE_LOG(LogNNE, Warning, TEXT("Runtime %s cannot create the model data with id %s (Filetype: %s)"), *RuntimeName , *FileId.ToString(EGuidFormats::Digits).ToLower(), *FileType);
+				UE_LOG(LogNNE, Warning, TEXT("Runtime %s cannot create the model data with id %s (Filetype: %s)"), *RuntimeName, *FileId.ToString(EGuidFormats::Digits).ToLower(), *FileType);
 			}
 		}
 		else
diff --git a/Engine/Source/Runtime/NNE/Public/NNERuntime.h b/Engine/Source/Runtime/NNE/Public/NNERuntime.h
--- a/Engine/Source/Runtime/NNE/Public/NNERuntime.h
+++ b/Engine/Source/Runtime/NNE/Public/NNERuntime.h
@@ -35,7 +35,12 @@ class NNE_API INNERuntime
 
 public:
 
-	using ECanCreateModelDataStatus = UE::NNE::EResultStatus;
+	enum class ECanCreateModelDataStatus : uint8
+	{
+		Ok = 0,
+		Fail = 1,
+		FailFileIdNotSupported = 2
+	};
 
 	/**
 	 * Get the name of the runtime.
