[NNERuntimeIREE] Added support for more complex meta data on imported model functions

Problem: MLIR model functions can contain attributes and other metadata in the function signature as well as the arguments, not all of those have been treated correctly by the parser.

Solution: Adding regex and more generic parsing of the function signatures.

Tested: By importing several models with complex signatures and metadata in the editor.

#rnx

diff --git a/Engine/Plugins/Experimental/NNERuntimeIREE/Source/NNERuntimeIREE/Private/NNERuntimeIREEMetaData.cpp b/Engine/Plugins/Experimental/NNERuntimeIREE/Source/NNERuntimeIREE/Private/NNERuntimeIREEMetaData.cpp
--- a/Engine/Plugins/Experimental/NNERuntimeIREE/Source/NNERuntimeIREE/Private/NNERuntimeIREEMetaData.cpp
+++ b/Engine/Plugins/Experimental/NNERuntimeIREE/Source/NNERuntimeIREE/Private/NNERuntimeIREEMetaData.cpp
@@ -3,8 +3,9 @@
 #include "NNERuntimeIREEMetaData.h"
 
 #include "Containers/Array.h"
-#include "Serialization/CustomVersion.h"
+#include "Internationalization/Regex.h"
 #include "NNE.h"
+#include "Serialization/CustomVersion.h"
 
 namespace UE::NNERuntimeIREE::ModuleMetaData::Private
 {
@@ -32,6 +33,10 @@ namespace UE::NNERuntimeIREE::ModuleMetaData::Private
 		{
 			return ENNETensorDataType::Half;
 		}
+		else if (TypeString.StartsWith("bf16"))
+		{
+			return ENNETensorDataType::BFloat16;
+		}
 		else if (TypeString.StartsWith("f"))
 		{
 			if (TypeString.StartsWith("f16"))
@@ -92,149 +97,158 @@ namespace UE::NNERuntimeIREE::ModuleMetaData::Private
 		return ENNETensorDataType::None;
 	}
 
-	bool ParseTensorDescFromString(const FString& ArgumentString, FString& Name, TArray<int32>& Shape, ENNETensorDataType& Type)
+	int32 FindCorrespondingClosingSymbol(const FString& String, int32 Offset, const FString& OpenSymbol, const FString& CloseSymbol)
 	{
-		FString Argument = ArgumentString;
-		int32 AttributeStartIndex = Argument.Find("{");
-		if (AttributeStartIndex > 0)
-		{
-			Argument = Argument.Mid(0, AttributeStartIndex).TrimStartAndEnd();
-		}
-
-		Name = "";
-		int32 NameEnd = Argument.Find(":");
-		if (NameEnd > 0)
-		{
-			Name = Argument.Mid(0, NameEnd).TrimStartAndEnd();
-			Argument = Argument.Mid(NameEnd + 1).TrimStartAndEnd();
-		}
-		else
+		int32 OpenedSymbols = 1;
+		int32 NextClose = String.Find(CloseSymbol, ESearchCase::IgnoreCase, ESearchDir::FromStart, Offset);
+		if (NextClose == INDEX_NONE)
 		{
-			Argument = Argument.TrimStartAndEnd();
+			return INDEX_NONE;
 		}
 
-		Shape.Empty();
-		if (Argument.StartsWith("tensor"))
+		int32 NextOpen = String.Find(OpenSymbol, ESearchCase::IgnoreCase, ESearchDir::FromStart, Offset);
+		while (OpenedSymbols > 0)
 		{
-			int32 ShapeOpenBracket = Argument.Find("<");
-			int32 ShapeCloseBracket = Argument.Find(">");
-			if (ShapeOpenBracket < 1 || ShapeCloseBracket <= ShapeOpenBracket)
+			if (NextOpen == INDEX_NONE || NextClose < NextOpen)
 			{
-				return false;
+				OpenedSymbols--;
+				if (OpenedSymbols > 0)
+				{
+					NextClose = String.Find(CloseSymbol, ESearchCase::IgnoreCase, ESearchDir::FromStart, NextClose + 1);
+					if (NextClose == INDEX_NONE)
+					{
+						return INDEX_NONE;
+					}
+				}
 			}
-			FString ShapeTypeString = Argument.Mid(ShapeOpenBracket + 1, ShapeCloseBracket - (ShapeOpenBracket + 1)).TrimStartAndEnd();
-
-			TArray<FString> Dims;
-			FString Dim;
-			while (ShapeTypeString.Split("x", &Dim, &ShapeTypeString))
+			else
 			{
-				Dims.Add(Dim);
+				OpenedSymbols++;
+				NextOpen = String.Find(OpenSymbol, ESearchCase::IgnoreCase, ESearchDir::FromStart, NextOpen + 1);
 			}
-			Type = ConvertTypeString(ShapeTypeString);
+		}
+		return NextClose;
+	}
 
-			for (int32 i = 0; i < Dims.Num(); i++)
+	bool ParseArgumentType(const FString& ArgumentType, UE::NNE::FSymbolicTensorShape& Shape, ENNETensorDataType& Type)
+	{
+		int32 ShapeStart = ArgumentType.Find("<");
+		if (ShapeStart != INDEX_NONE)
+		{
+			int32 ShapeEnd = ArgumentType.Find(">", ESearchCase::IgnoreCase, ESearchDir::FromStart, ShapeStart);
+			check(ShapeStart != INDEX_NONE);
+			FString ShapeString = ArgumentType.Mid(ShapeStart + 1, ShapeEnd - ShapeStart - 1).TrimStartAndEnd();
+			TArray<FString> ShapeList;
+			ShapeString.ParseIntoArray(ShapeList, TEXT(","));
+			if (ShapeList.Num() < 2)
+			{
+				ShapeString.ParseIntoArray(ShapeList, TEXT("x"));
+			}
+			TArray<int32> ShapeArray;
+			for (int32 i = 0; i < ShapeList.Num() - 1; i++)
 			{
-				if (Dims[i].Contains("?"))
+				if (ShapeList[i].Contains("?"))
 				{
-					Shape.Add(-1);
+					ShapeArray.Emplace(-1);
 				}
 				else
 				{
-					int32 DimVal = FCString::Atoi(*Dims[i]);
-					Shape.Add(DimVal > 0 ? DimVal : -1);
+					// Match any integer possibly starting with a sign
+					// [-+]?: Matches 0 or 1 or optional times plus or minus
+					// \\d+: Matches one or more digits
+					const FRegexPattern Pattern(TEXT("[-+]?\\d+"));
+					FRegexMatcher Matcher(Pattern, ShapeList[i]);
+					if (Matcher.FindNext())
+					{
+						int32 IntegerStart = Matcher.GetMatchBeginning();
+						int32 IntegerEnd = Matcher.GetMatchEnding();
+						check(IntegerStart != INDEX_NONE && IntegerEnd != INDEX_NONE);
+						ShapeArray.Emplace(FCString::Atoi(*ShapeList[i].Mid(IntegerStart, IntegerEnd - IntegerStart)));
+					}
 				}
+
 			}
+			Shape = UE::NNE::FSymbolicTensorShape::Make(ShapeArray);
+			Type = ConvertTypeString(ShapeList[ShapeList.Num() - 1]);
 		}
 		else
 		{
-			Shape.Add(1);
-			Type = ConvertTypeString(Argument);
+			Type = ConvertTypeString(ArgumentType);
 		}
-
-		return Type != ENNETensorDataType::None;
+		return true;
 	}
 
-	bool ParseFunctionMetaDataFromString(const FString& FunctionString, UE::NNERuntimeIREE::FFunctionMetaData& FunctionMetaData)
+	bool ParseArguments(const FString& Arguments, TArray<UE::NNE::FTensorDesc>& TensorDescs)
 	{
-		FString ResultPattern = "->";
-		int32 ArgumentsOpenBracket = FunctionString.Find("(");
-		int32 ArgumentsCloseBracket = FunctionString.Find(")");
-		int32 ResultStart = FunctionString.Find(ResultPattern);
-		if (ArgumentsOpenBracket < 1 || ArgumentsCloseBracket <= ArgumentsOpenBracket)
+		FString ReducedArguments = "";
+		int32 LastStart = 0;
+		int32 NextOpenSymbol = Arguments.Find("{", ESearchCase::IgnoreCase, ESearchDir::FromStart, LastStart);
+		while (NextOpenSymbol != INDEX_NONE)
 		{
-			return false;
+			ReducedArguments += Arguments.Mid(LastStart, NextOpenSymbol - LastStart);
+			LastStart = FindCorrespondingClosingSymbol(Arguments, NextOpenSymbol + 1, "{", "}") + 1;
+			check(LastStart != INDEX_NONE);
+			NextOpenSymbol = Arguments.Find("{", ESearchCase::IgnoreCase, ESearchDir::FromStart, LastStart);
 		}
+		ReducedArguments += Arguments.Mid(LastStart);
 
-		FunctionMetaData.Name = FunctionString.Mid(0, ArgumentsOpenBracket).TrimStartAndEnd();
-		if (FunctionMetaData.Name.IsEmpty())
+		FString FinalArguments = "";
+		LastStart = 0;
+		NextOpenSymbol = ReducedArguments.Find("(", ESearchCase::IgnoreCase, ESearchDir::FromStart, LastStart);
+		while (NextOpenSymbol != INDEX_NONE)
 		{
-			return false;
+			FinalArguments += ReducedArguments.Mid(LastStart, NextOpenSymbol - LastStart);
+			LastStart = FindCorrespondingClosingSymbol(ReducedArguments, NextOpenSymbol + 1, "(", ")") + 1;
+			check(LastStart != INDEX_NONE);
+			NextOpenSymbol = ReducedArguments.Find("(", ESearchCase::IgnoreCase, ESearchDir::FromStart, LastStart);
 		}
-		FunctionMetaData.InputDescs.Empty();
-		FunctionMetaData.OutputDescs.Empty();
+		FinalArguments += ReducedArguments.Mid(LastStart);
 
-		FString ArgumentsString = FunctionString.Mid(ArgumentsOpenBracket + 1, ArgumentsCloseBracket - (ArgumentsOpenBracket + 1)).TrimStartAndEnd();
-		while (ArgumentsString.Len() > 0)
+		TArray<FString> ArgumentList;
+		int32 LastIndex = 0;
+		int32 InsideShape = 0;
+		for (int32 i = 0; i < FinalArguments.Len(); i++)
 		{
-			int32 SeparatorIndex = ArgumentsString.Find(",");
-			FString ArgumentString;
-			if (SeparatorIndex > 0)
-			{
-				ArgumentString = ArgumentsString.Mid(0, SeparatorIndex).TrimStartAndEnd();
-				ArgumentsString = ArgumentsString.Mid(SeparatorIndex + 1).TrimStartAndEnd();
-			}
-			else
+			if (FinalArguments[i] == *TEXT(",") && InsideShape == 0)
 			{
-				ArgumentString = ArgumentsString;
-				ArgumentsString = "";
+				ArgumentList.Add(FinalArguments.Mid(LastIndex, i - LastIndex).TrimStartAndEnd());
+				LastIndex = i + 1;
 			}
-
-			FString Name;
-			TArray<int32> Shape;
-			ENNETensorDataType Type;
-			if (ParseTensorDescFromString(ArgumentString, Name, Shape, Type))
+			else if (FinalArguments[i] == *TEXT("<"))
 			{
-				FunctionMetaData.InputDescs.Add(UE::NNE::FTensorDesc::Make(Name, UE::NNE::FSymbolicTensorShape::Make(Shape), Type));
+				InsideShape++;
 			}
-			else
+			else if (FinalArguments[i] == *TEXT(">"))
 			{
-				return false;
+				InsideShape--;
 			}
 		}
+		ArgumentList.Add(FinalArguments.Mid(LastIndex).TrimStartAndEnd());
 
-		if (ResultStart > ArgumentsCloseBracket)
+		for (const FString& Argument : ArgumentList)
 		{
-			FString ResultsString = FunctionString.Mid(ResultStart + ResultPattern.Len()).TrimStartAndEnd();
-			ResultsString.RemoveFromStart("(");
-			ResultsString.RemoveFromEnd(")");
-			ResultsString = ResultsString.TrimStartAndEnd();
-			while (ResultsString.Len() > 0)
-			{
-				int32 SeparatorIndex = ResultsString.Find(",");
-				FString ResultString;
-				if (SeparatorIndex > 0)
-				{
-					ResultString = ResultsString.Mid(0, SeparatorIndex).TrimStartAndEnd();
-					ResultsString = ResultsString.Mid(SeparatorIndex + 1).TrimStartAndEnd();
-				}
-				else
-				{
-					ResultString = ResultsString;
-					ResultsString = "";
-				}
+			TArray<FString> ArgumentPartList;
+			Argument.ParseIntoArray(ArgumentPartList, TEXT(":"));
 
-				FString Name;
-				TArray<int32> Shape;
-				ENNETensorDataType Type;
-				if (ParseTensorDescFromString(ResultString, Name, Shape, Type))
+			FString Name;
+			UE::NNE::FSymbolicTensorShape Shape;
+			ENNETensorDataType Type = ENNETensorDataType::None;
+			if (ArgumentPartList.Num() > 1)
+			{
+				Name = ArgumentPartList[0].TrimStartAndEnd();
+				if (!ParseArgumentType(ArgumentPartList[1].TrimStartAndEnd(), Shape, Type))
 				{
-					FunctionMetaData.OutputDescs.Add(UE::NNE::FTensorDesc::Make(Name, UE::NNE::FSymbolicTensorShape::Make(Shape), Type));
+					return false;
 				}
-				else
+			}
+			else
+			{
+				if (!ParseArgumentType(ArgumentPartList[0].TrimStartAndEnd(), Shape, Type))
 				{
 					return false;
 				}
 			}
+			TensorDescs.Add(UE::NNE::FTensorDesc::Make(Name, Shape, Type));
 		}
 
 		return true;
@@ -337,46 +351,80 @@ bool UNNERuntimeIREEModuleMetaData::ParseFromString(const FString& ModuleString)
 	using namespace UE::NNERuntimeIREE::ModuleMetaData::Private;
 
 	TArray<UE::NNERuntimeIREE::FFunctionMetaData> Result;
-	FString SearchString = ModuleString;
-	FString Pattern = "func.func";
-	int32 MatchStart = -1;
-	while ((MatchStart = SearchString.Find(Pattern)) > 0)
+
+	// Match func.func [access modifier] @<function_name> (
+	// func\.func: Exact match of the word
+	// [^@]*: Skip zero or more characters not being @
+	// @: Have exactly one @
+	// [^(]*: Skip zero or more characters not being (
+	// (: Have exactly one (
+	const FRegexPattern FunctionNamePattern(TEXT("func\\.func[^@]*@[^(]*\\("));
+	FRegexMatcher FunctionNameMatcher(FunctionNamePattern, *ModuleString);
+	while (FunctionNameMatcher.FindNext())
 	{
-		MatchStart += Pattern.Len();
-		SearchString = SearchString.Mid(MatchStart).TrimStartAndEnd();
-		if (SearchString.StartsWith("@"))
+		int32 FunctionStart = FunctionNameMatcher.GetMatchBeginning();
+		int32 InputArgumentsStart = FunctionNameMatcher.GetMatchEnding();
+		check(FunctionStart != INDEX_NONE && InputArgumentsStart != INDEX_NONE);
+
+		FString MatchedFunctionPattern = ModuleString.Mid(FunctionStart, InputArgumentsStart - FunctionStart - 1);
+		bool bIsPrivate = MatchedFunctionPattern.Find("private") != INDEX_NONE;
+		if (bIsPrivate)
 		{
-			SearchString = SearchString.Mid(1);
+			//continue;
+		}
+		int32 FunctionNameStart = MatchedFunctionPattern.Find("@");
+		check(FunctionNameStart != INDEX_NONE);
+
+		UE::NNERuntimeIREE::FFunctionMetaData MetaData;
+		MetaData.Name = MatchedFunctionPattern.Mid(FunctionNameStart+1).TrimStartAndEnd();
+		check(MetaData.Name.Len() > 0);
+
+		// Regex does not support recursive matching but mlir can contain paranthesis inside arguments
+		int32 InputArgumentsEnd = FindCorrespondingClosingSymbol(ModuleString, InputArgumentsStart, "(", ")");
+		check(InputArgumentsEnd != INDEX_NONE);
+
+		FString InputArguments = ModuleString.Mid(InputArgumentsStart, InputArgumentsEnd - InputArgumentsStart);
+		if (!ParseArguments(InputArguments, MetaData.InputDescs))
+		{
+			return false;
+		}
 
-			int32 ArgumentsEnd = SearchString.Find(")");
-			FString TempString = SearchString.Mid(ArgumentsEnd + 1).TrimStartAndEnd();
-			if (TempString.StartsWith("->"))
+		// Match [white spaces] -> [white spaces]
+		// ^: Match start of the string to make sure that there are any non white space characters at the beginning
+		// [\\s]*: Match zero or more white spaces
+		// ->: Exact match of ->
+		// [\\s]*: Match zero or more white spaces
+		const FRegexPattern OutputArgumentsStartPattern(TEXT("^[\\s]*->[\\s]*"));
+		FString Rest = ModuleString.Mid(InputArgumentsEnd+1);
+		FRegexMatcher OutputArgumentsStartMatcher(OutputArgumentsStartPattern, Rest);
+		if (OutputArgumentsStartMatcher.FindNext())
+		{
+			int32 OutputArgumentsStart = OutputArgumentsStartMatcher.GetMatchEnding();
+			int32 OutputArgumentsEnd = INDEX_NONE;
+			Rest = Rest.Mid(OutputArgumentsStart);
+			if (Rest.StartsWith("("))
 			{
-				TempString = TempString.Mid(3).TrimStartAndEnd();
-				if (TempString.StartsWith("("))
-				{
-					MatchStart = SearchString.Find(")", ESearchCase::IgnoreCase, ESearchDir::FromStart, ArgumentsEnd + 1);
-				}
-				else
-				{
-					int32 AttribtueStart = SearchString.Find("attributes", ESearchCase::IgnoreCase, ESearchDir::FromStart, ArgumentsEnd + 1);
-					int32 BracketStart = SearchString.Find("{", ESearchCase::IgnoreCase, ESearchDir::FromStart, ArgumentsEnd + 1);
-					MatchStart = AttribtueStart <= ArgumentsEnd ? BracketStart : FMath::Min(AttribtueStart, BracketStart);
-				}
+				OutputArgumentsStart = 1;
+				OutputArgumentsEnd = FindCorrespondingClosingSymbol(Rest, OutputArgumentsStart, "(", ")");
 			}
 			else
 			{
-				MatchStart = ArgumentsEnd;
+				OutputArgumentsStart = 0;
+				int32 ClosestParanthesis = Rest.Find("(");
+				int32 ClosestBraces = Rest.Find("{");
+				OutputArgumentsEnd = (ClosestParanthesis == INDEX_NONE || ClosestBraces < ClosestParanthesis) ? ClosestBraces : ClosestParanthesis;
 			}
-
-			UE::NNERuntimeIREE::FFunctionMetaData MetaData;
-			if (!ParseFunctionMetaDataFromString(SearchString.Mid(0, MatchStart), MetaData))
+			check(OutputArgumentsEnd != INDEX_NONE);
+			FString OutputArguments = Rest.Mid(OutputArgumentsStart, OutputArgumentsEnd - OutputArgumentsStart);
+			if (!ParseArguments(OutputArguments, MetaData.OutputDescs))
 			{
 				return false;
 			}
-			Result.Add(MetaData);
 		}
+
+		Result.Add(MetaData);
 	}
+	
 	if (!Result.IsEmpty())
 	{
 		FunctionMetaData = Result;
