[NNERuntimeIREE] Don't list private module functions

Problem: Currently all functions inside an mlir module are listed when importing it

Solution: Filter meta data parsing to export only public functions which can be accessed later

Tested by importing a model with both public and private functions.

#rnx

diff --git a/Engine/Plugins/Experimental/NNERuntimeIREE/Source/NNERuntimeIREE/Private/NNERuntimeIREEMetaData.cpp b/Engine/Plugins/Experimental/NNERuntimeIREE/Source/NNERuntimeIREE/Private/NNERuntimeIREEMetaData.cpp
--- a/Engine/Plugins/Experimental/NNERuntimeIREE/Source/NNERuntimeIREE/Private/NNERuntimeIREEMetaData.cpp
+++ b/Engine/Plugins/Experimental/NNERuntimeIREE/Source/NNERuntimeIREE/Private/NNERuntimeIREEMetaData.cpp
@@ -368,9 +368,10 @@ bool UNNERuntimeIREEModuleMetaData::ParseFromString(const FString& ModuleString)
 
 		FString MatchedFunctionPattern = ModuleString.Mid(FunctionStart, InputArgumentsStart - FunctionStart - 1);
 		bool bIsPrivate = MatchedFunctionPattern.Find("private") != INDEX_NONE;
-		if (bIsPrivate)
+		bool bIsProtected = MatchedFunctionPattern.Find("protected") != INDEX_NONE;
+		if (bIsPrivate || bIsProtected)
 		{
-			//continue;
+			continue;
 		}
 		int32 FunctionNameStart = MatchedFunctionPattern.Find("@");
 		check(FunctionNameStart != INDEX_NONE);
