[NNE] NNERuntimeIREE fixed config update and verification.

#tests Editor, Win64 verified config creation/update and later correct validation

#rnx

diff --git a/Engine/Plugins/Experimental/NNERuntimeIREE/Source/NNERuntimeIREE/Private/NNERuntimeIREE.cpp b/Engine/Plugins/Experimental/NNERuntimeIREE/Source/NNERuntimeIREE/Private/NNERuntimeIREE.cpp
--- a/Engine/Plugins/Experimental/NNERuntimeIREE/Source/NNERuntimeIREE/Private/NNERuntimeIREE.cpp
+++ b/Engine/Plugins/Experimental/NNERuntimeIREE/Source/NNERuntimeIREE/Private/NNERuntimeIREE.cpp
@@ -127,10 +127,8 @@ TSharedPtr<UE::NNE::FSharedModelData> UNNERuntimeIREECpu::CreateModelData(const
 		return TSharedPtr<UE::NNE::FSharedModelData>();
 	}
 
-	FConfigFile ConfigFile;
 	FString ConfigFilePath;
-	GetUpdatedPlatformConfig(TargetPlatformName, ConfigFile, ConfigFilePath);
-	if (ConfigFile.Dirty)
+	if (!UpdatePlatformConfigAndVerify(TargetPlatformName, true, &ConfigFilePath))
 	{
 		UE_LOG(LogNNE, Warning, TEXT("UNNERuntimeIREECpu could not find the required settings in config file %s. Please make the file writeable and re-start the editor or manually add the required staging settings or models will not work in packaged builds for platform %s!"), *ConfigFilePath, *TargetPlatformName);
 	}
@@ -396,19 +394,53 @@ TSharedPtr<UE::NNE::IModelCPU> UNNERuntimeIREECpu::CreateModelCPU(const TObjectP
 	return Model;
 }
 
-void UNNERuntimeIREECpu::GetUpdatedPlatformConfig(const FString& PlatformName, FConfigFile& ConfigFile, FString& ConfigFilePath)
-{ 
-	FString ConfigFolderPath = FPaths::ConvertRelativePathToFull(FPaths::ProjectConfigDir());
-	ConfigFilePath = FPaths::Combine(ConfigFolderPath, PlatformName, PlatformName + "Game.ini");
+bool UNNERuntimeIREECpu::UpdatePlatformConfigAndVerify(const FString& PlatformName, bool bOnlyVerify, FString* ConfigFilePathPtr)
+{
+	const FString ConfigFolderPath = FPaths::ConvertRelativePathToFull(FPaths::ProjectConfigDir());
+	const FString ConfigFilePath = FPaths::Combine(ConfigFolderPath, PlatformName, PlatformName + "Game.ini");
+
+	if (ConfigFilePathPtr)
+	{
+		*ConfigFilePathPtr = ConfigFilePath;
+	}
+
+	const FString StagingPath = FString("/") + UE::NNERuntimeIREE::CPU::Private::GetStagedModelDirPath(PlatformName);
+	const FString PackagingPath = FString("/") + UE::NNERuntimeIREE::CPU::Private::GetPackagedModelDirPath(PlatformName);
+
+	const FString DirectoriesToAlwaysStageAsNonUFAddition = FString("(Path=\"..") + StagingPath + FString("\")");
+	const FString RemapDirectoriesAddition = FString("(From=\"") + FApp::GetProjectName() + StagingPath + FString("\", To=\"") + FApp::GetProjectName() + PackagingPath + FString("\")");
+	const FString AllowedDirectoriesAddition = FApp::GetProjectName() + PackagingPath;
+
+	if (!bOnlyVerify)
+	{
+		FConfigFile ConfigFile;
+		ConfigFile.Read(ConfigFilePath);
+
+		bool bChanged = false;
+		bChanged |= ConfigFile.AddUniqueToSection(TEXT("/Script/UnrealEd.ProjectPackagingSettings"), TEXT("+DirectoriesToAlwaysStageAsNonUFS"), DirectoriesToAlwaysStageAsNonUFAddition);
+		bChanged |= ConfigFile.AddUniqueToSection(TEXT("Staging"), TEXT("+RemapDirectories"), RemapDirectoriesAddition);
+		bChanged |= ConfigFile.AddUniqueToSection(TEXT("Staging"), TEXT("+AllowedDirectories"), AllowedDirectoriesAddition);
+
+		if (bChanged)
+		{
+			ConfigFile.Write(ConfigFilePath);
+		}
+	}
+
+	auto CheckConfigArrayElement = [] (const TCHAR *Section, const TCHAR *Key, const FString& Value, const FString& ConfigFilePath)
+	{
+		TArray<FString> Values;
+		GConfig->GetArray(Section, Key, Values, ConfigFilePath);
 
-	ConfigFile.Read(ConfigFilePath);
+		return Values.Contains(Value);
+	};
 
-	FString StagingPath = FString("/") + UE::NNERuntimeIREE::CPU::Private::GetStagedModelDirPath(PlatformName);
-	FString PackagingPath = FString("/") + UE::NNERuntimeIREE::CPU::Private::GetPackagedModelDirPath(PlatformName);
+	bool Result = true;
+	Result &= CheckConfigArrayElement(TEXT("/Script/UnrealEd.ProjectPackagingSettings"), TEXT("+DirectoriesToAlwaysStageAsNonUFS"), DirectoriesToAlwaysStageAsNonUFAddition, ConfigFilePath);
+	Result &= CheckConfigArrayElement(TEXT("Staging"), TEXT("+RemapDirectories"), RemapDirectoriesAddition, ConfigFilePath);
+	Result &= CheckConfigArrayElement(TEXT("Staging"), TEXT("+AllowedDirectories"), AllowedDirectoriesAddition, ConfigFilePath);
 
-	ConfigFile.AddUniqueToSection(TEXT("/Script/UnrealEd.ProjectPackagingSettings"), TEXT("+DirectoriesToAlwaysStageAsNonUFS"), FString("(Path=\"..") + StagingPath + FString("\")"));
-	ConfigFile.AddUniqueToSection(TEXT("Staging"), TEXT("+RemapDirectories"), FString("(From=\"") + FApp::GetProjectName() + StagingPath + FString("\", To=\"") + FApp::GetProjectName() + PackagingPath + FString("\")"));
-	ConfigFile.AddUniqueToSection(TEXT("Staging"), TEXT("+AllowedDirectories"), FApp::GetProjectName() + PackagingPath);
+	return Result;
 }
 
 FString UNNERuntimeIREEGpu::GetRuntimeName() const
diff --git a/Engine/Plugins/Experimental/NNERuntimeIREE/Source/NNERuntimeIREE/Private/NNERuntimeIREE.h b/Engine/Plugins/Experimental/NNERuntimeIREE/Source/NNERuntimeIREE/Private/NNERuntimeIREE.h
--- a/Engine/Plugins/Experimental/NNERuntimeIREE/Source/NNERuntimeIREE/Private/NNERuntimeIREE.h
+++ b/Engine/Plugins/Experimental/NNERuntimeIREE/Source/NNERuntimeIREE/Private/NNERuntimeIREE.h
@@ -34,7 +34,7 @@ public:
 	virtual TSharedPtr<UE::NNE::IModelCPU> CreateModelCPU(const TObjectPtr<UNNEModelData> ModelData) override;
 	//~ End INNERuntimeCPU Interface
 
-	static void GetUpdatedPlatformConfig(const FString& PlatformName, FConfigFile& ConfigFile, FString& ConfigFilePath);
+	static bool UpdatePlatformConfigAndVerify(const FString& PlatformName, bool bOnlyVerify, FString* ConfigFilePathPtr = nullptr);
 #else
 	//~ Begin INNERuntime Interface
 	virtual FString GetRuntimeName() const override { return ""; };
diff --git a/Engine/Plugins/Experimental/NNERuntimeIREE/Source/NNERuntimeIREE/Private/NNERuntimeIREEModule.cpp b/Engine/Plugins/Experimental/NNERuntimeIREE/Source/NNERuntimeIREE/Private/NNERuntimeIREEModule.cpp
--- a/Engine/Plugins/Experimental/NNERuntimeIREE/Source/NNERuntimeIREE/Private/NNERuntimeIREEModule.cpp
+++ b/Engine/Plugins/Experimental/NNERuntimeIREE/Source/NNERuntimeIREE/Private/NNERuntimeIREEModule.cpp
@@ -18,6 +18,35 @@
 
 #endif // WITH_NNE_RUNTIME_IREE
 
+#if WITH_EDITOR
+namespace UE::NNERuntimeIREE::ConfigHelper
+{
+
+void CpuUpdatePlatformConfigs()
+{
+	ITargetPlatformManagerModule* TargetPlatformManagerModule = GetTargetPlatformManager();
+	if (!TargetPlatformManagerModule)
+	{
+		return;
+	}
+
+	TSet<FString> ProcessedPlatforms;
+	TArray<ITargetPlatform*> TargetPlatforms = TargetPlatformManagerModule->GetTargetPlatforms();
+	for (int32 i = 0; i < TargetPlatforms.Num(); i++)
+	{
+		FString IniPlatformName = TargetPlatforms[i]->IniPlatformName();
+		if (ProcessedPlatforms.Contains(IniPlatformName))
+		{
+			continue;
+		}
+
+		UNNERuntimeIREECpu::UpdatePlatformConfigAndVerify(IniPlatformName, false);
+	}
+}
+
+}
+#endif // WITH_EDITOR
+
 void FNNERuntimeIREEModule::StartupModule()
 {
 #ifdef WITH_NNE_RUNTIME_IREE
@@ -29,35 +58,7 @@ void FNNERuntimeIREEModule::StartupModule()
 	}
 
 #if WITH_EDITOR
-	ITargetPlatformManagerModule* TargetPlatformManagerModule = GetTargetPlatformManager();
-	if (TargetPlatformManagerModule)
-	{
-		TSet<FString> ProcessedPlatforms;
-		TArray<ITargetPlatform*> TargetPlatforms = TargetPlatformManagerModule->GetTargetPlatforms();
-		for (int32 i = 0; i < TargetPlatforms.Num(); i++)
-		{
-			FString IniPlatformName = TargetPlatforms[i]->IniPlatformName();
-			if (!ProcessedPlatforms.Contains(IniPlatformName))
-			{
-				ProcessedPlatforms.Add(IniPlatformName);
-				FString TargetPlatformDisplayName = IniPlatformName;
-				FConfigFile ConfigFile;
-				FString ConfigFilePath;
-				NNERuntimeIREECpu->GetUpdatedPlatformConfig(IniPlatformName, ConfigFile, ConfigFilePath);
-				if (ConfigFile.Dirty)
-				{
-					{
-						TUniquePtr<FArchive> Ar = TUniquePtr<FArchive>(IFileManager::Get().CreateFileWriter(*ConfigFilePath, EFileWrite::FILEWRITE_Append));
-						if (!Ar)
-						{
-							continue;
-						}
-					}
-					ConfigFile.Write(ConfigFilePath);
-				}
-			}
-		}
-	}
+	UE::NNERuntimeIREE::ConfigHelper::CpuUpdatePlatformConfigs();
 #endif // WITH_EDITOR
 
 	NNERuntimeIREECuda = NewObject<UNNERuntimeIREECuda>();
