[NNE] NNERuntimeIREE log cleanup.

#tests Editor, Win64 startup, PIE

#rnx

diff --git a/Engine/Plugins/Experimental/NNERuntimeIREE/Source/NNERuntimeIREE/Internal/NNERuntimeIREELog.h b/Engine/Plugins/Experimental/NNERuntimeIREE/Source/NNERuntimeIREE/Internal/NNERuntimeIREELog.h
new file mode 100644
--- /dev/null
+++ b/Engine/Plugins/Experimental/NNERuntimeIREE/Source/NNERuntimeIREE/Internal/NNERuntimeIREELog.h
@@ -0,0 +1,7 @@
+// Copyright Epic Games, Inc. All Rights Reserved.
+
+#pragma once
+
+#include "Logging/LogMacros.h"
+
+DECLARE_LOG_CATEGORY_EXTERN(LogNNERuntimeIREE, Log, All);
\ No newline at end of file
diff --git a/Engine/Plugins/Experimental/NNERuntimeIREE/Source/NNERuntimeIREE/Private/NNERuntimeIREE.cpp b/Engine/Plugins/Experimental/NNERuntimeIREE/Source/NNERuntimeIREE/Private/NNERuntimeIREE.cpp
--- a/Engine/Plugins/Experimental/NNERuntimeIREE/Source/NNERuntimeIREE/Private/NNERuntimeIREE.cpp
+++ b/Engine/Plugins/Experimental/NNERuntimeIREE/Source/NNERuntimeIREE/Private/NNERuntimeIREE.cpp
@@ -22,10 +22,11 @@
 #include "Misc/Paths.h"
 #include "Misc/SecureHash.h"
 #include "NNE.h"
+#include "NNEModelData.h"
 #include "NNERuntimeIREECompiler.h"
-#include "NNERuntimeIREEModel.h"
+#include "NNERuntimeIREELog.h"
 #include "NNERuntimeIREEMetaData.h"
-#include "NNEModelData.h"
+#include "NNERuntimeIREEModel.h"
 #include "Serialization/MemoryReader.h"
 #include "Serialization/MemoryWriter.h"
 
@@ -123,20 +124,20 @@ TSharedPtr<UE::NNE::FSharedModelData> UNNERuntimeIREECpu::CreateModelData(const
 	FString TargetPlatformName = TargetPlatform ? TargetPlatform->IniPlatformName() : UGameplayStatics::GetPlatformName();
 	if (CanCreateModelData(FileType, FileData, AdditionalFileData, FileId, TargetPlatform) != ECanCreateModelDataStatus::Ok)
 	{
-		UE_LOG(LogNNE, Warning, TEXT("UNNERuntimeIREECpu cannot create the model data with id %s (Filetype: %s) for platform %s"), *FileId.ToString(EGuidFormats::Digits).ToLower(), *FileType, *TargetPlatformName);
+		UE_LOG(LogNNERuntimeIREE, Warning, TEXT("UNNERuntimeIREECpu cannot create the model data with id %s (Filetype: %s) for platform %s"), *FileId.ToString(EGuidFormats::Digits).ToLower(), *FileType, *TargetPlatformName);
 		return TSharedPtr<UE::NNE::FSharedModelData>();
 	}
 
 	FString ConfigFilePath;
 	if (!UpdatePlatformConfigAndVerify(TargetPlatformName, true, &ConfigFilePath))
 	{
-		UE_LOG(LogNNE, Warning, TEXT("UNNERuntimeIREECpu could not find the required settings in config file %s. Please make the file writeable and re-start the editor or manually add the required staging settings or models will not work in packaged builds for platform %s!"), *ConfigFilePath, *TargetPlatformName);
+		UE_LOG(LogNNERuntimeIREE, Warning, TEXT("UNNERuntimeIREECpu could not find the required settings in config file %s. Please make the file writeable and re-start the editor or manually add the required staging settings or models will not work in packaged builds for platform %s!"), *ConfigFilePath, *TargetPlatformName);
 	}
 
 	TUniquePtr<UE::NNERuntimeIREE::CPU::FCompiler> Compiler = UE::NNERuntimeIREE::CPU::FCompiler::Make(TargetPlatformName);
 	if (!Compiler.IsValid())
 	{
-		UE_LOG(LogNNE, Warning, TEXT("UNNERuntimeIREECpu failed to create a compiler to compile for platform %s"), *TargetPlatformName);
+		UE_LOG(LogNNERuntimeIREE, Warning, TEXT("UNNERuntimeIREECpu failed to create a compiler to compile for platform %s"), *TargetPlatformName);
 		return TSharedPtr<UE::NNE::FSharedModelData>();
 	}
 	
@@ -151,7 +152,7 @@ TSharedPtr<UE::NNE::FSharedModelData> UNNERuntimeIREECpu::CreateModelData(const
 	FString StagingDir = FPaths::ConvertRelativePathToFull(FPaths::Combine(FPaths::ProjectDir(), GetStagedModelDirPath(TargetPlatformName)));
 	if (!Compiler->CompileMlir(FileData, FileIdString, IntermediateDir, StagingDir, CompilerResults, CompilerModuleMetaData))
 	{
-		UE_LOG(LogNNE, Warning, TEXT("UNNERuntimeIREECpu failed to compile model %s"), *FileIdString);
+		UE_LOG(LogNNERuntimeIREE, Warning, TEXT("UNNERuntimeIREECpu failed to compile model %s"), *FileIdString);
 		return TSharedPtr<UE::NNE::FSharedModelData>();
 	}
 
@@ -161,7 +162,7 @@ TSharedPtr<UE::NNE::FSharedModelData> UNNERuntimeIREECpu::CreateModelData(const
 		FString StagedSharedLibPath = FPaths::Combine(StagingDir, CompilerResults[i].RelativeDirPath, CompilerResults[i].SharedLibraryFileName);
 		if (!FFileHelper::LoadFileToArray(SharedLibData, *StagedSharedLibPath) || SharedLibData.IsEmpty())
 		{
-			UE_LOG(LogNNE, Warning, TEXT("UNNERuntimeIREECpu could not read the shared library \"%s\""), *StagedSharedLibPath);
+			UE_LOG(LogNNERuntimeIREE, Warning, TEXT("UNNERuntimeIREECpu could not read the shared library \"%s\""), *StagedSharedLibPath);
 			return TSharedPtr<UE::NNE::FSharedModelData>();
 		}
 		FSharedBuffer SharedLibBuffer = MakeSharedBufferFromArray(MoveTemp(SharedLibData));
@@ -171,7 +172,7 @@ TSharedPtr<UE::NNE::FSharedModelData> UNNERuntimeIREECpu::CreateModelData(const
 		FString StagedVmfbPath = FPaths::Combine(StagingDir, CompilerResults[i].RelativeDirPath, CompilerResults[i].VmfbFileName);
 		if (!FFileHelper::LoadFileToArray(VmfbData, *StagedVmfbPath) || VmfbData.IsEmpty())
 		{
-			UE_LOG(LogNNE, Warning, TEXT("UNNERuntimeIREECpu could not read the vmfb data \"%s\""), *StagedVmfbPath);
+			UE_LOG(LogNNERuntimeIREE, Warning, TEXT("UNNERuntimeIREECpu could not read the vmfb data \"%s\""), *StagedVmfbPath);
 			return TSharedPtr<UE::NNE::FSharedModelData>();
 		}
 		FSharedBuffer VmfbBuffer = MakeSharedBufferFromArray(MoveTemp(VmfbData));
@@ -253,7 +254,7 @@ TSharedPtr<UE::NNE::IModelCPU> UNNERuntimeIREECpu::CreateModelCPU(const TObjectP
 
 	if (CanCreateModelCPU(ModelData) != ECanCreateModelCPUStatus::Ok)
 	{
-		UE_LOG(LogNNE, Warning, TEXT("UNNERuntimeIREECpu cannot create a model from the model data with id %s"), *ModelData->GetFileId().ToString(EGuidFormats::Digits));
+		UE_LOG(LogNNERuntimeIREE, Warning, TEXT("UNNERuntimeIREECpu cannot create a model from the model data with id %s"), *ModelData->GetFileId().ToString(EGuidFormats::Digits));
 		return TSharedPtr<UE::NNE::IModelCPU>();
 	}
 
@@ -279,7 +280,7 @@ TSharedPtr<UE::NNE::IModelCPU> UNNERuntimeIREECpu::CreateModelCPU(const TObjectP
 	Reader << ModuleDataArray;
 	if (ModuleDataArray.IsEmpty())
 	{
-		UE_LOG(LogNNE, Warning, TEXT("UNNERuntimeIREECpu failed to find any module meta data, please reimport the original model"));
+		UE_LOG(LogNNERuntimeIREE, Warning, TEXT("UNNERuntimeIREECpu failed to find any module meta data, please reimport the original model"));
 		return TSharedPtr<UE::NNE::IModelCPU>();
 	}
 	UNNERuntimeIREEModuleMetaData* ModuleMetaData = NewObject<UNNERuntimeIREEModuleMetaData>();
@@ -287,7 +288,7 @@ TSharedPtr<UE::NNE::IModelCPU> UNNERuntimeIREECpu::CreateModelCPU(const TObjectP
 	ModuleMetaData->Serialize(ObjectReader);
 	if (ModuleMetaData->FunctionMetaData.IsEmpty())
 	{
-		UE_LOG(LogNNE, Warning, TEXT("UNNERuntimeIREECpu failed to parse the module meta data, please reimport the original model"));
+		UE_LOG(LogNNERuntimeIREE, Warning, TEXT("UNNERuntimeIREECpu failed to parse the module meta data, please reimport the original model"));
 		return TSharedPtr<UE::NNE::IModelCPU>();
 	}
 
@@ -334,7 +335,7 @@ TSharedPtr<UE::NNE::IModelCPU> UNNERuntimeIREECpu::CreateModelCPU(const TObjectP
 	}
 	if (!bFound)
 	{
-		UE_LOG(LogNNE, Warning, TEXT("UNNERuntimeIREECpu failed to find a matching architecture for \'%s\'"), *CurrentArchitecture);
+		UE_LOG(LogNNERuntimeIREE, Warning, TEXT("UNNERuntimeIREECpu failed to find a matching architecture for \'%s\'"), *CurrentArchitecture);
 		return TSharedPtr<UE::NNE::IModelCPU>();
 	}
 
@@ -353,7 +354,7 @@ TSharedPtr<UE::NNE::IModelCPU> UNNERuntimeIREECpu::CreateModelCPU(const TObjectP
 		FSharedBuffer SharedBuffer = GetFromDDC(UE::NNERuntimeIREE::CPU::Private::GetModelDataIdentifier(GetRuntimeName(), UNNERuntimeIREECpu::GUID, FileIdString, UGameplayStatics::GetPlatformName(), Architecture) + "-lib");
 		if (SharedBuffer.GetSize() <= 0)
 		{
-			UE_LOG(LogNNE, Warning, TEXT("UNNERuntimeIREECpu could not fetch the shared library %s from DDC"), *SharedLibraryFileName);
+			UE_LOG(LogNNERuntimeIREE, Warning, TEXT("UNNERuntimeIREECpu could not fetch the shared library %s from DDC"), *SharedLibraryFileName);
 			return TSharedPtr<UE::NNE::IModelCPU>();
 		}
 		PlatformFile.CreateDirectoryTree(*SharedLibraryDirPath);
@@ -366,7 +367,7 @@ TSharedPtr<UE::NNE::IModelCPU> UNNERuntimeIREECpu::CreateModelCPU(const TObjectP
 		FSharedBuffer SharedBuffer = GetFromDDC(UE::NNERuntimeIREE::CPU::Private::GetModelDataIdentifier(GetRuntimeName(), UNNERuntimeIREECpu::GUID, FileIdString, UGameplayStatics::GetPlatformName(), Architecture) + "-vmfb");
 		if (SharedBuffer.GetSize() <= 0)
 		{
-			UE_LOG(LogNNE, Warning, TEXT("UNNERuntimeIREECpu could not fetch the vmfb %s from DDC"), *VmfbFileName);
+			UE_LOG(LogNNERuntimeIREE, Warning, TEXT("UNNERuntimeIREECpu could not fetch the vmfb %s from DDC"), *VmfbFileName);
 			return TSharedPtr<UE::NNE::IModelCPU>();
 		}
 		PlatformFile.CreateDirectoryTree(*SharedLibraryDirPath);
@@ -377,7 +378,7 @@ TSharedPtr<UE::NNE::IModelCPU> UNNERuntimeIREECpu::CreateModelCPU(const TObjectP
 	TSharedPtr<UE::NNE::IModelCPU> Model = UE::NNERuntimeIREE::CPU::FModel::Make(SharedLibraryDirPath, SharedLibraryFileName, VmfbFileName, SharedLibraryEntryPointName, *ModuleMetaData);
 	if (!Model.IsValid())
 	{
-		UE_LOG(LogNNE, Warning, TEXT("UNNERuntimeIREECpu could not initialize the model created from model data with id %s"), *FileIdString);
+		UE_LOG(LogNNERuntimeIREE, Warning, TEXT("UNNERuntimeIREECpu could not initialize the model created from model data with id %s"), *FileIdString);
 		return TSharedPtr<UE::NNE::IModelCPU>();
 	}
 
@@ -500,7 +501,7 @@ TSharedPtr<UE::NNE::IModelGPU> UNNERuntimeIREEGpu::CreateModelGPU(const TObjectP
 
 	if (CanCreateModelGPU(ModelData) != ECanCreateModelGPUStatus::Ok)
 	{
-		UE_LOG(LogNNE, Warning, TEXT("UNNERuntimeIREEGpu cannot create a model from the model data with id %s"), *ModelData->GetFileId().ToString(EGuidFormats::Digits));
+		UE_LOG(LogNNERuntimeIREE, Warning, TEXT("UNNERuntimeIREEGpu cannot create a model from the model data with id %s"), *ModelData->GetFileId().ToString(EGuidFormats::Digits));
 		return TSharedPtr<UE::NNE::IModelGPU>();
 	}
 
@@ -641,7 +642,7 @@ TSharedPtr<UE::NNE::IModelRDG> UNNERuntimeIREERdg::CreateModelRDG(const TObjectP
 
 	if (CanCreateModelRDG(ModelData) != ECanCreateModelRDGStatus::Ok)
 	{
-		UE_LOG(LogNNE, Warning, TEXT("UNNERuntimeIREERdg cannot create a model from the model data with id %s"), *ModelData->GetFileId().ToString(EGuidFormats::Digits));
+		UE_LOG(LogNNERuntimeIREE, Warning, TEXT("UNNERuntimeIREERdg cannot create a model from the model data with id %s"), *ModelData->GetFileId().ToString(EGuidFormats::Digits));
 		return TSharedPtr<UE::NNE::IModelRDG>();
 	}
 
diff --git a/Engine/Plugins/Experimental/NNERuntimeIREE/Source/NNERuntimeIREE/Private/NNERuntimeIREECompiler.cpp b/Engine/Plugins/Experimental/NNERuntimeIREE/Source/NNERuntimeIREE/Private/NNERuntimeIREECompiler.cpp
--- a/Engine/Plugins/Experimental/NNERuntimeIREE/Source/NNERuntimeIREE/Private/NNERuntimeIREECompiler.cpp
+++ b/Engine/Plugins/Experimental/NNERuntimeIREE/Source/NNERuntimeIREE/Private/NNERuntimeIREECompiler.cpp
@@ -15,6 +15,7 @@
 #include "Misc/FileHelper.h"
 #include "Misc/Paths.h"
 #include "NNE.h"
+#include "NNERuntimeIREELog.h"
 #include "Serialization/JsonReader.h"
 #include "Serialization/JsonSerializer.h"
 #include "Templates/SharedPointer.h"
@@ -125,7 +126,7 @@ namespace UE::NNERuntimeIREE
 						{
 							if (BuildConfig.BuildTargets.IsEmpty())
 							{
-								UE_LOG(LogNNE, Warning, TEXT("UNNERuntimeIREECpu could not find targets in %s"), *BuildConfigFilePath);
+								UE_LOG(LogNNERuntimeIREE, Warning, TEXT("UNNERuntimeIREECpu could not find targets in %s"), *BuildConfigFilePath);
 								continue;
 							}
 
@@ -144,12 +145,12 @@ namespace UE::NNERuntimeIREE
 								}
 								else
 								{
-									UE_LOG(LogNNE, Warning, TEXT("UNNERuntimeIREECpu could not replace environment variables in %s"), *BuildConfig.CompilerCommand[i]);
+									UE_LOG(LogNNERuntimeIREE, Warning, TEXT("UNNERuntimeIREECpu could not replace environment variables in %s"), *BuildConfig.CompilerCommand[i]);
 								}
 							}
 							if (TmpCompilerCommand.IsEmpty())
 							{
-								UE_LOG(LogNNE, Warning, TEXT("UNNERuntimeIREECpu could not find the compiler executable in %s"), *BuildConfigFilePath);
+								UE_LOG(LogNNERuntimeIREE, Warning, TEXT("UNNERuntimeIREECpu could not find the compiler executable in %s"), *BuildConfigFilePath);
 								continue;
 							}
 
@@ -168,12 +169,12 @@ namespace UE::NNERuntimeIREE
 								}
 								else
 								{
-									UE_LOG(LogNNE, Warning, TEXT("UNNERuntimeIREECpu could not replace environment variables in %s"), *BuildConfig.LinkerCommand[i]);
+									UE_LOG(LogNNERuntimeIREE, Warning, TEXT("UNNERuntimeIREECpu could not replace environment variables in %s"), *BuildConfig.LinkerCommand[i]);
 								}
 							}
 							if (TmpLinkerCommand.IsEmpty())
 							{
-								UE_LOG(LogNNE, Warning, TEXT("UNNERuntimeIREECpu could not find the linker executable in %s"), *BuildConfigFilePath);
+								UE_LOG(LogNNERuntimeIREE, Warning, TEXT("UNNERuntimeIREECpu could not find the linker executable in %s"), *BuildConfigFilePath);
 								continue;
 							}
 
@@ -185,12 +186,12 @@ namespace UE::NNERuntimeIREE
 						}
 						else
 						{
-							UE_LOG(LogNNE, Warning, TEXT("UNNERuntimeIREECpu could not parse build config file %s"), *BuildConfigFilePath);
+							UE_LOG(LogNNERuntimeIREE, Warning, TEXT("UNNERuntimeIREECpu could not parse build config file %s"), *BuildConfigFilePath);
 						}
 					}
 					else
 					{
-						UE_LOG(LogNNE, Warning, TEXT("UNNERuntimeIREECpu could not read build config file %s"), *BuildConfigFilePath);
+						UE_LOG(LogNNERuntimeIREE, Warning, TEXT("UNNERuntimeIREECpu could not read build config file %s"), *BuildConfigFilePath);
 					}
 				}
 			}
@@ -230,7 +231,7 @@ namespace UE::NNERuntimeIREE
 				FString CompilerArguments = BuildTargets[i].CompilerArguments;
 				if (!ResolveEnvironmentVariables(CompilerArguments))
 				{
-					UE_LOG(LogNNE, Warning, TEXT("UNNERuntimeIREECpu could not replace environment variables in %s"), *BuildTargets[i].CompilerArguments);
+					UE_LOG(LogNNERuntimeIREE, Warning, TEXT("UNNERuntimeIREECpu could not replace environment variables in %s"), *BuildTargets[i].CompilerArguments);
 					continue;
 				}
 				CompilerArguments.ReplaceInline(*FString("${OBJECT_PATH}"), *(FString("\"") + ObjectFilePath + "\""));
@@ -241,8 +242,8 @@ namespace UE::NNERuntimeIREE
 
 				if (!PlatformFile.FileExists(*ObjectFilePath) || !PlatformFile.FileExists(*VmfbFilePath))
 				{
-					UE_LOG(LogNNE, Warning, TEXT("UNNERuntimeIREECpu failed to compile the model \"%s\" using the command:"), *InputFilePath);
-					UE_LOG(LogNNE, Warning, TEXT("\"%s\" %s"), *CompilerCommand, *CompilerArguments);
+					UE_LOG(LogNNERuntimeIREE, Warning, TEXT("UNNERuntimeIREECpu failed to compile the model \"%s\" using the command:"), *InputFilePath);
+					UE_LOG(LogNNERuntimeIREE, Warning, TEXT("\"%s\" %s"), *CompilerCommand, *CompilerArguments);
 					bResult = false;
 					continue;
 				}
@@ -250,7 +251,7 @@ namespace UE::NNERuntimeIREE
 				FString LinkerArguments = BuildTargets[i].LinkerArguments;
 				if (!ResolveEnvironmentVariables(LinkerArguments))
 				{
-					UE_LOG(LogNNE, Warning, TEXT("UNNERuntimeIREECpu could not replace environment variables in %s"), *BuildTargets[i].LinkerArguments);
+					UE_LOG(LogNNERuntimeIREE, Warning, TEXT("UNNERuntimeIREECpu could not replace environment variables in %s"), *BuildTargets[i].LinkerArguments);
 					bResult = false;
 					continue;
 				}
@@ -261,8 +262,8 @@ namespace UE::NNERuntimeIREE
 
 				if (!PlatformFile.FileExists(*SharedLibFilePath))
 				{
-					UE_LOG(LogNNE, Warning, TEXT("UNNERuntimeIREECpu failed to link the model \"%s\" using the command:"), *InputFilePath);
-					UE_LOG(LogNNE, Warning, TEXT("\"%s\" %s"), *LinkerCommand, *LinkerArguments);
+					UE_LOG(LogNNERuntimeIREE, Warning, TEXT("UNNERuntimeIREECpu failed to link the model \"%s\" using the command:"), *InputFilePath);
+					UE_LOG(LogNNERuntimeIREE, Warning, TEXT("\"%s\" %s"), *LinkerCommand, *LinkerArguments);
 					bResult = false;
 					continue;
 				}
@@ -272,12 +273,12 @@ namespace UE::NNERuntimeIREE
 				FString StagedSharedLibFilePath = StagingFilePathNoExt + SharedLibExt;
 				if (IFileManager::Get().Copy(*StagedSharedLibFilePath, *SharedLibFilePath) != COPY_OK)
 				{
-					UE_LOG(LogNNE, Warning, TEXT("UNNERuntimeIREECpu failed to copy \"%s\" to \"%s\""), *SharedLibFilePath, *StagedSharedLibFilePath);
+					UE_LOG(LogNNERuntimeIREE, Warning, TEXT("UNNERuntimeIREECpu failed to copy \"%s\" to \"%s\""), *SharedLibFilePath, *StagedSharedLibFilePath);
 					bResult = false;
 				}
 				if (IFileManager::Get().Copy(*StagedVmfbFilePath, *VmfbFilePath) != COPY_OK)
 				{
-					UE_LOG(LogNNE, Warning, TEXT("UNNERuntimeIREECpu failed to copy \"%s\" to \"%s\""), *VmfbFilePath, *StagedVmfbFilePath);
+					UE_LOG(LogNNERuntimeIREE, Warning, TEXT("UNNERuntimeIREECpu failed to copy \"%s\" to \"%s\""), *VmfbFilePath, *StagedVmfbFilePath);
 					bResult = false;
 				}
 
@@ -285,21 +286,21 @@ namespace UE::NNERuntimeIREE
 				FString HeaderPath = IntermediateFilePathNoExt + ".h";
 				if (!PlatformFile.FileExists(*HeaderPath))
 				{
-					UE_LOG(LogNNE, Warning, TEXT("UNNERuntimeIREECpu could not find the model header \"%s\""), *HeaderPath);
+					UE_LOG(LogNNERuntimeIREE, Warning, TEXT("UNNERuntimeIREECpu could not find the model header \"%s\""), *HeaderPath);
 					bResult = false;
 					continue;
 				}
 				FString HeaderString;
 				if (!FFileHelper::LoadFileToString(HeaderString, *HeaderPath) || HeaderString.IsEmpty())
 				{
-					UE_LOG(LogNNE, Warning, TEXT("UNNERuntimeIREECpu could not read the model header \"%s\""), *HeaderPath);
+					UE_LOG(LogNNERuntimeIREE, Warning, TEXT("UNNERuntimeIREECpu could not read the model header \"%s\""), *HeaderPath);
 					bResult = false;
 					continue;
 				}
 				SharedLibraryEntryPointName = GetSharedLibraryEntryPointName(HeaderString);
 				if (SharedLibraryEntryPointName.IsEmpty())
 				{
-					UE_LOG(LogNNE, Warning, TEXT("UNNERuntimeIREECpu could not find the entry point in model header \"%s\""), *HeaderPath);
+					UE_LOG(LogNNERuntimeIREE, Warning, TEXT("UNNERuntimeIREECpu could not find the entry point in model header \"%s\""), *HeaderPath);
 					bResult = false;
 					continue;
 				}
diff --git a/Engine/Plugins/Experimental/NNERuntimeIREE/Source/NNERuntimeIREE/Private/NNERuntimeIREEMetaData.cpp b/Engine/Plugins/Experimental/NNERuntimeIREE/Source/NNERuntimeIREE/Private/NNERuntimeIREEMetaData.cpp
--- a/Engine/Plugins/Experimental/NNERuntimeIREE/Source/NNERuntimeIREE/Private/NNERuntimeIREEMetaData.cpp
+++ b/Engine/Plugins/Experimental/NNERuntimeIREE/Source/NNERuntimeIREE/Private/NNERuntimeIREEMetaData.cpp
@@ -5,6 +5,7 @@
 #include "Containers/Array.h"
 #include "Internationalization/Regex.h"
 #include "NNE.h"
+#include "NNERuntimeIREELog.h"
 #include "Serialization/CustomVersion.h"
 
 namespace UE::NNERuntimeIREE::ModuleMetaData::Private
@@ -340,7 +341,7 @@ void UNNERuntimeIREEModuleMetaData::Serialize(FArchive& Ar)
 			}
 			break;
 		default:
-			UE_LOG(LogNNE, Error, TEXT("UNNERuntimeIREEModuleMetaData: Unknown asset version %d: Deserialisation failed, please reimport the original model."), Ar.CustomVer(UE::NNERuntimeIREE::ModuleMetaData::Private::GUID));
+			UE_LOG(LogNNERuntimeIREE, Error, TEXT("UNNERuntimeIREEModuleMetaData: Unknown asset version %d: Deserialisation failed, please reimport the original model."), Ar.CustomVer(UE::NNERuntimeIREE::ModuleMetaData::Private::GUID));
 			break;
 		}
 	}
diff --git a/Engine/Plugins/Experimental/NNERuntimeIREE/Source/NNERuntimeIREE/Private/NNERuntimeIREEModel.cpp b/Engine/Plugins/Experimental/NNERuntimeIREE/Source/NNERuntimeIREE/Private/NNERuntimeIREEModel.cpp
--- a/Engine/Plugins/Experimental/NNERuntimeIREE/Source/NNERuntimeIREE/Private/NNERuntimeIREEModel.cpp
+++ b/Engine/Plugins/Experimental/NNERuntimeIREE/Source/NNERuntimeIREE/Private/NNERuntimeIREEModel.cpp
@@ -11,6 +11,7 @@
 #include "Misc/ScopeLock.h"
 #include "NNE.h"
 #include "NNEModelData.h"
+#include "NNERuntimeIREELog.h"
 #include "NNEStatus.h"
 #include "Serialization/Archive.h"
 
@@ -103,7 +104,7 @@ namespace UE::NNERuntimeIREE
 			void* ErrorString = FMemory::Malloc(TrueLength + 1);
 			((char*)ErrorString)[TrueLength] = (char)0;
 			iree_status_format(InStatus, TrueLength, (char*)ErrorString, &TrueLength);
-			UE_LOG(LogNNE, Error, TEXT("%s: %s"), *InMessage, *FString(StringCast<TCHAR>(static_cast<const ANSICHAR*>(ErrorString)).Get()));
+			UE_LOG(LogNNERuntimeIREE, Error, TEXT("%s: %s"), *InMessage, *FString(StringCast<TCHAR>(static_cast<const ANSICHAR*>(ErrorString)).Get()));
 			FMemory::Free(ErrorString);
 		}
 
@@ -353,13 +354,13 @@ namespace UE::NNERuntimeIREE
 				TUniquePtr<FArchive> Reader = TUniquePtr<FArchive>(IFileManager::Get().CreateFileReader(*CombinedPath, 0));
 				if (!Reader)
 				{
-					UE_LOG(LogNNE, Error, TEXT("UE::NNERuntimeIREE::Private::FModule failed to open the vmfb data file '%s'"), *CombinedPath);
+					UE_LOG(LogNNERuntimeIREE, Error, TEXT("UE::NNERuntimeIREE::Private::FModule failed to open the vmfb data file '%s'"), *CombinedPath);
 					return TSharedPtr<FModule>();
 				}
 				int64 DataSize = Reader->TotalSize();
 				if (DataSize < 1)
 				{
-					UE_LOG(LogNNE, Error, TEXT("UE::NNERuntimeIREE::Private::FModule's vmfb data file '%s' is empty"), *CombinedPath);
+					UE_LOG(LogNNERuntimeIREE, Error, TEXT("UE::NNERuntimeIREE::Private::FModule's vmfb data file '%s' is empty"), *CombinedPath);
 					return TSharedPtr<FModule>();
 				}
 
@@ -438,7 +439,7 @@ namespace UE::NNERuntimeIREE
 
 				if (!bFound)
 				{
-					UE_LOG(LogNNE, Error, TEXT("UE::NNERuntimeIREE::Private::FModule failed to find the module function %s"), *InFunctionName);
+					UE_LOG(LogNNERuntimeIREE, Error, TEXT("UE::NNERuntimeIREE::Private::FModule failed to find the module function %s"), *InFunctionName);
 					iree_status_free(Status);
 					return false;
 				}
@@ -486,7 +487,7 @@ namespace UE::NNERuntimeIREE
 					void* Library = FPlatformProcess::GetDllHandle(*CombinedPath);
 					if (!Library)
 					{
-						UE_LOG(LogNNE, Error, TEXT("UE::NNERuntimeIREE::CPU::Private::FLibrary failed to load the shared library '%s'"), *CombinedPath);
+						UE_LOG(LogNNERuntimeIREE, Error, TEXT("UE::NNERuntimeIREE::CPU::Private::FLibrary failed to load the shared library '%s'"), *CombinedPath);
 						return TSharedPtr<FLibrary>();
 					}
 #else
@@ -495,7 +496,7 @@ namespace UE::NNERuntimeIREE
 					FPlatformProcess::PopDllDirectory(*InLibraryPath);
 					if (!Library)
 					{
-						UE_LOG(LogNNE, Error, TEXT("UE::NNERuntimeIREE::CPU::Private::FLibrary failed to load the shared library '%s' from '%s'"), *InLibraryName, *InLibraryPath);
+						UE_LOG(LogNNERuntimeIREE, Error, TEXT("UE::NNERuntimeIREE::CPU::Private::FLibrary failed to load the shared library '%s' from '%s'"), *InLibraryName, *InLibraryPath);
 						return TSharedPtr<FLibrary>();
 					}
 #endif
@@ -513,7 +514,7 @@ namespace UE::NNERuntimeIREE
 						*OutFunctionPointer = Result;
 						return true;
 					}
-					UE_LOG(LogNNE, Error, TEXT("UE::NNERuntimeIREE::CPU::Private::FLibrary failed to get the function %s"), *InFunctionName);
+					UE_LOG(LogNNERuntimeIREE, Error, TEXT("UE::NNERuntimeIREE::CPU::Private::FLibrary failed to get the function %s"), *InFunctionName);
 					return false;
 				}
 			};
@@ -701,7 +702,7 @@ namespace UE::NNERuntimeIREE
 					TConstArrayView<UE::NNE::FTensorDesc> OutputTensorDescs = InModule->GetFunctionMetaDataView()[0].OutputDescs;
 					if (!iree_status_is_ok(Status) || NumInputs != InputTensorDescs.Num() || NumOutputs != OutputTensorDescs.Num())
 					{
-						UE_LOG(LogNNE, Error, TEXT("UE::NNERuntimeIREE::CPU::Private::FSession has a function signature mismatch in function %s"), *MainFunctionName);
+						UE_LOG(LogNNERuntimeIREE, Error, TEXT("UE::NNERuntimeIREE::CPU::Private::FSession has a function signature mismatch in function %s"), *MainFunctionName);
 						iree_runtime_session_release(Session);
 						iree_status_free(Status);
 						return TSharedPtr<FSession>();
@@ -844,7 +845,7 @@ namespace UE::NNERuntimeIREE
 							iree_hal_buffer_t* Buffer = iree_hal_buffer_view_buffer(BufferViews[i]);
 							if (!Buffer)
 							{
-								UE_LOG(LogNNE, Error, TEXT("UE::NNERuntimeIREE::CPU::Private::FSession failed to get the result buffer"));
+								UE_LOG(LogNNERuntimeIREE, Error, TEXT("UE::NNERuntimeIREE::CPU::Private::FSession failed to get the result buffer"));
 								Result = ERunSyncStatus::Fail;
 								break;
 							}
diff --git a/Engine/Plugins/Experimental/NNERuntimeIREE/Source/NNERuntimeIREE/Private/NNERuntimeIREEModule.cpp b/Engine/Plugins/Experimental/NNERuntimeIREE/Source/NNERuntimeIREE/Private/NNERuntimeIREEModule.cpp
--- a/Engine/Plugins/Experimental/NNERuntimeIREE/Source/NNERuntimeIREE/Private/NNERuntimeIREEModule.cpp
+++ b/Engine/Plugins/Experimental/NNERuntimeIREE/Source/NNERuntimeIREE/Private/NNERuntimeIREEModule.cpp
@@ -18,6 +18,10 @@
 
 #endif // WITH_NNE_RUNTIME_IREE
 
+#include "NNERuntimeIREELog.h"
+
+DEFINE_LOG_CATEGORY(LogNNERuntimeIREE);
+
 #if WITH_EDITOR
 namespace UE::NNERuntimeIREE::ConfigHelper
 {
diff --git a/Engine/Plugins/Experimental/NNERuntimeIREE/Source/NNERuntimeIREEEditor/Private/NNERuntimeIREEEditorModule.cpp b/Engine/Plugins/Experimental/NNERuntimeIREE/Source/NNERuntimeIREEEditor/Private/NNERuntimeIREEEditorModule.cpp
--- a/Engine/Plugins/Experimental/NNERuntimeIREE/Source/NNERuntimeIREEEditor/Private/NNERuntimeIREEEditorModule.cpp
+++ b/Engine/Plugins/Experimental/NNERuntimeIREE/Source/NNERuntimeIREEEditor/Private/NNERuntimeIREEEditorModule.cpp
@@ -1,9 +1,12 @@
 // Copyright Epic Games, Inc. All Rights Reserved.
 
 #include "NNERuntimeIREEEditorModule.h"
+#include "NNERuntimeIREELog.h"
 
 #define LOCTEXT_NAMESPACE "FNNERuntimeIREEEditorModule"
 
+DEFINE_LOG_CATEGORY(LogNNERuntimeIREE);
+
 void FNNERuntimeIREEEditorModule::StartupModule()
 {
 }
diff --git a/Engine/Plugins/Experimental/NNERuntimeIREE/Source/NNERuntimeIREEEditor/Private/NNERuntimeIREEModelDataFactory.cpp b/Engine/Plugins/Experimental/NNERuntimeIREE/Source/NNERuntimeIREEEditor/Private/NNERuntimeIREEModelDataFactory.cpp
--- a/Engine/Plugins/Experimental/NNERuntimeIREE/Source/NNERuntimeIREEEditor/Private/NNERuntimeIREEModelDataFactory.cpp
+++ b/Engine/Plugins/Experimental/NNERuntimeIREE/Source/NNERuntimeIREEEditor/Private/NNERuntimeIREEModelDataFactory.cpp
@@ -14,6 +14,7 @@
 #include "Modules/ModuleManager.h"
 #include "NNE.h"
 #include "NNEModelData.h"
+#include "NNERuntimeIREELog.h"
 #include "NNERuntimeIREEMetaData.h"
 #include "Serialization/MemoryWriter.h"
 #include "Subsystems/ImportSubsystem.h"
@@ -164,7 +165,7 @@ UObject* UNNERuntimeIREEModelDataFactory::FactoryCreateBinary(UClass* Class, UOb
 	TObjectPtr<UNNERuntimeIREEModuleMetaData> ModuleMetaData = NewObject<UNNERuntimeIREEModuleMetaData>();
 	if (!ModuleMetaData->ParseFromString(FileDataString) || ModuleMetaData->FunctionMetaData.IsEmpty())
 	{
-		UE_LOG(LogNNE, Error, TEXT("UNNERuntimeIREEModelDataFactory failed to parse the models meta data"));
+		UE_LOG(LogNNERuntimeIREE, Error, TEXT("UNNERuntimeIREEModelDataFactory failed to parse the models meta data"));
 		return nullptr;
 	}
 
@@ -194,7 +195,7 @@ UObject* UNNERuntimeIREEModelDataFactory::FactoryCreateBinary(UClass* Class, UOb
 
 		if (!ImportWindow->ImportButtonClicked())
 		{
-			UE_LOG(LogNNE, Error, TEXT("UNNERuntimeIREEModelDataFactory could not import the model! Please select the neural network main function in the import dialog!"));
+			UE_LOG(LogNNERuntimeIREE, Error, TEXT("UNNERuntimeIREEModelDataFactory could not import the model! Please select the neural network main function in the import dialog!"));
 			return nullptr;
 		}
 
