[NNERuntimeIREE] In PIE IREE compiler might hang with some models (specifically with thrown exceptions).

#tests Editor, Win64 run PIE with the same model that caused hang before

#rnx

diff --git a/Engine/Plugins/Experimental/NNERuntimeIREE/Source/NNERuntimeIREE/Private/NNERuntimeIREECompiler.cpp b/Engine/Plugins/Experimental/NNERuntimeIREE/Source/NNERuntimeIREE/Private/NNERuntimeIREECompiler.cpp
--- a/Engine/Plugins/Experimental/NNERuntimeIREE/Source/NNERuntimeIREE/Private/NNERuntimeIREECompiler.cpp
+++ b/Engine/Plugins/Experimental/NNERuntimeIREE/Source/NNERuntimeIREE/Private/NNERuntimeIREECompiler.cpp
@@ -13,6 +13,7 @@
 #include "Interfaces/IPluginManager.h"
 #include "Kismet/GameplayStatics.h"
 #include "Misc/FileHelper.h"
+#include "Misc/MonitoredProcess.h"
 #include "Misc/Paths.h"
 #include "NNE.h"
 #include "NNERuntimeIREELog.h"
@@ -76,15 +77,41 @@ namespace UE::NNERuntimeIREE
 				return HeaderString.Mid(Start, End - Start).TrimStartAndEnd();
 			}
 
-			void RunCommand(const FString& Command, const FString& Arguments)
+			void RunCommand(const FString& Command, const FString& Arguments, const FString& LogFilePath = FString())
 			{
-				void* PipeRead = nullptr;
-				void* PipeWrite = nullptr;
-				FPlatformProcess::CreatePipe(PipeRead, PipeWrite);
-				FProcHandle ProcHandle = FPlatformProcess::CreateProc(*Command, *Arguments, false, true, true, nullptr, 0, nullptr, PipeWrite, PipeRead);
-				FPlatformProcess::WaitForProc(ProcHandle);
-				FPlatformProcess::CloseProc(ProcHandle);
-				FPlatformProcess::ClosePipe(PipeRead, PipeWrite);
+				int32 ReturnCode = 0;
+				bool IsCanceled = false;
+
+				FMonitoredProcess Process(Command, Arguments, true);
+				Process.OnCompleted().BindLambda([&ReturnCode] (int32 _ReturnCode) { ReturnCode = _ReturnCode; });
+				Process.OnCanceled().BindLambda([&IsCanceled] (){ IsCanceled = true; });
+
+				if (!Process.Launch())
+				{
+					UE_LOG(LogNNERuntimeIREE, Warning, TEXT("Failed to launch subprocess!"));
+					return;
+				}
+
+				while (Process.Update())
+				{
+					// Poll until process has finished
+				}
+
+				if (IsCanceled)
+				{
+					UE_LOG(LogNNERuntimeIREE, Warning, TEXT("Execution of subprocess was canceled!"));
+				}
+				else if (ReturnCode)
+				{
+					UE_LOG(LogNNERuntimeIREE, Warning, TEXT("Subprocess exited with non-zero code %d"), ReturnCode);
+				}
+
+				if (!LogFilePath.IsEmpty())
+				{
+					FFileHelper::SaveStringToFile(Process.GetFullOutputWithoutDelegate(), *LogFilePath);
+
+					UE_LOG(LogNNERuntimeIREE, Log, TEXT("Saved subprocess output to: %s"), *LogFilePath);
+				}
 			}
 		} // Private
 
@@ -238,7 +265,7 @@ namespace UE::NNERuntimeIREE
 				CompilerArguments.ReplaceInline(*FString("${VMFB_PATH}"), *(FString("\"") + VmfbFilePath + "\""));
 				CompilerArguments.ReplaceInline(*FString("${INPUT_PATH}"), *(FString("\"") + InputFilePath + "\""));
 
-				RunCommand(CompilerCommand, CompilerArguments);
+				RunCommand(CompilerCommand, CompilerArguments, IntermediateFilePathNoExt + "_compile-log.txt");
 
 				if (!PlatformFile.FileExists(*ObjectFilePath) || !PlatformFile.FileExists(*VmfbFilePath))
 				{
@@ -258,7 +285,7 @@ namespace UE::NNERuntimeIREE
 				LinkerArguments.ReplaceInline(*FString("${OBJECT_PATH}"), *(FString("\"") + ObjectFilePath + "\""));
 				LinkerArguments.ReplaceInline(*FString("${SHARED_LIB_PATH}"), *(FString("\"") + SharedLibFilePath + "\""));
 
-				RunCommand(LinkerCommand, LinkerArguments);
+				RunCommand(LinkerCommand, LinkerArguments, IntermediateFilePathNoExt + "_link-log.txt");
 
 				if (!PlatformFile.FileExists(*SharedLibFilePath))
 				{
