[NNE] NNERuntimeIREE revise plugin config.

#tests Packaged and ran on all platforms

#rnx

diff --git a/Engine/Plugins/Experimental/NNERuntimeIREE/Config/Linux/LinuxGame.ini b/Engine/Plugins/Experimental/NNERuntimeIREE/Config/Linux/LinuxGame.ini
new file mode 100644
--- /dev/null
+++ b/Engine/Plugins/Experimental/NNERuntimeIREE/Config/Linux/LinuxGame.ini
@@ -0,0 +1,2 @@
+[/Script/UnrealEd.ProjectPackagingSettings]
++DirectoriesToAlwaysStageAsNonUFS=(Path="../Binaries/Linux/NNERuntimeIREE")
diff --git a/Engine/Plugins/Experimental/NNERuntimeIREE/Config/LinuxArm64/LinuxArm64Game.ini b/Engine/Plugins/Experimental/NNERuntimeIREE/Config/LinuxArm64/LinuxArm64Game.ini
new file mode 100644
--- /dev/null
+++ b/Engine/Plugins/Experimental/NNERuntimeIREE/Config/LinuxArm64/LinuxArm64Game.ini
@@ -0,0 +1,2 @@
+[/Script/UnrealEd.ProjectPackagingSettings]
++DirectoriesToAlwaysStageAsNonUFS=(Path="../Binaries/LinuxArm64/NNERuntimeIREE")
diff --git a/Engine/Plugins/Experimental/NNERuntimeIREE/Config/Mac/MacGame.ini b/Engine/Plugins/Experimental/NNERuntimeIREE/Config/Mac/MacGame.ini
new file mode 100644
--- /dev/null
+++ b/Engine/Plugins/Experimental/NNERuntimeIREE/Config/Mac/MacGame.ini
@@ -0,0 +1,2 @@
+[/Script/UnrealEd.ProjectPackagingSettings]
++DirectoriesToAlwaysStageAsNonUFS=(Path="../Binaries/Mac/NNERuntimeIREE")
diff --git a/Engine/Plugins/Experimental/NNERuntimeIREE/Config/Windows/WindowsGame.ini b/Engine/Plugins/Experimental/NNERuntimeIREE/Config/Windows/WindowsGame.ini
new file mode 100644
--- /dev/null
+++ b/Engine/Plugins/Experimental/NNERuntimeIREE/Config/Windows/WindowsGame.ini
@@ -0,0 +1,2 @@
+[/Script/UnrealEd.ProjectPackagingSettings]
++DirectoriesToAlwaysStageAsNonUFS=(Path="../Binaries/Win64/NNERuntimeIREE")
diff --git a/Engine/Plugins/Experimental/NNERuntimeIREE/Source/NNERuntimeIREE/Private/NNERuntimeIREE.cpp b/Engine/Plugins/Experimental/NNERuntimeIREE/Source/NNERuntimeIREE/Private/NNERuntimeIREE.cpp
--- a/Engine/Plugins/Experimental/NNERuntimeIREE/Source/NNERuntimeIREE/Private/NNERuntimeIREE.cpp
+++ b/Engine/Plugins/Experimental/NNERuntimeIREE/Source/NNERuntimeIREE/Private/NNERuntimeIREE.cpp
@@ -89,13 +89,14 @@ namespace UE::NNERuntimeIREE::CPU::Private
 
 	FString GetStagedModelDirPath(const FString& PlatformName)
 	{
-		return FPaths::Combine("Saved", "Cooked", PlatformName, "Engine", "Plugins", UE_PLUGIN_NAME, "Binaries");
+		FString PlatformNameShort = PlatformName.Equals("Windows") ? "Win64" : PlatformName;
+
+		return FPaths::Combine("Binaries", PlatformNameShort, UE_PLUGIN_NAME);
 	}
 
 	FString GetPackagedModelDirPath(const FString& PlatformName)
 	{
-		FString PlatformNameShort = PlatformName.Equals("Windows") ? "Win64" : PlatformName;
-		return FPaths::Combine("Binaries", PlatformNameShort, UE_PLUGIN_NAME);
+		return GetStagedModelDirPath(PlatformName);
 	}
 } // UE::NNERuntimeIREE::CPU::Private
 
@@ -128,14 +129,6 @@ TSharedPtr<UE::NNE::FSharedModelData> UNNERuntimeIREECpu::CreateModelData(const
 		return TSharedPtr<UE::NNE::FSharedModelData>();
 	}
 
-	FConfigFile ConfigFile;
-	FString ConfigFilePath;
-	GetUpdatedPlatformConfig(TargetPlatformName, ConfigFile, ConfigFilePath);
-	if (ConfigFile.Dirty)
-	{
-		UE_LOG(LogNNERuntimeIREE, Warning, TEXT("UNNERuntimeIREECpu could not find the required settings in config file %s. Please make the file writeable and re-start the editor or manually add the required staging settings or models will not work in packaged builds for platform %s!"), *ConfigFilePath, *TargetPlatformName);
-	}
-
 	TUniquePtr<UE::NNERuntimeIREE::CPU::FCompiler> Compiler = UE::NNERuntimeIREE::CPU::FCompiler::Make(TargetPlatformName);
 	if (!Compiler.IsValid())
 	{
@@ -151,7 +144,7 @@ TSharedPtr<UE::NNE::FSharedModelData> UNNERuntimeIREECpu::CreateModelData(const
 
 	TArray<UE::NNERuntimeIREE::CPU::FCompilerResult> CompilerResults;
 	UNNERuntimeIREEModuleMetaData* CompilerModuleMetaData = NewObject<UNNERuntimeIREEModuleMetaData>();
-	FString StagingDir = FPaths::ConvertRelativePathToFull(FPaths::Combine(FPaths::ProjectDir(), GetStagedModelDirPath(TargetPlatformName)));
+	FString StagingDir = FPaths::ConvertRelativePathToFull(FPaths::Combine(FPaths::ProjectDir(), GetPackagedModelDirPath(TargetPlatformName)));
 	if (!Compiler->CompileMlir(FileData, FileIdString, IntermediateDir, StagingDir, CompilerResults, CompilerModuleMetaData))
 	{
 		UE_LOG(LogNNERuntimeIREE, Warning, TEXT("UNNERuntimeIREECpu failed to compile model %s"), *FileIdString);
@@ -397,21 +390,6 @@ TSharedPtr<UE::NNE::IModelCPU> UNNERuntimeIREECpu::CreateModelCPU(const TObjectP
 	return Model;
 }
 
-void UNNERuntimeIREECpu::GetUpdatedPlatformConfig(const FString& PlatformName, FConfigFile& ConfigFile, FString& ConfigFilePath)
-{ 
-	FString ConfigFolderPath = FPaths::ConvertRelativePathToFull(FPaths::ProjectConfigDir());
-	ConfigFilePath = FPaths::Combine(ConfigFolderPath, PlatformName, PlatformName + "Game.ini");
-
-	ConfigFile.Read(ConfigFilePath);
-
-	FString StagingPath = FString("/") + UE::NNERuntimeIREE::CPU::Private::GetStagedModelDirPath(PlatformName);
-	FString PackagingPath = FString("/") + UE::NNERuntimeIREE::CPU::Private::GetPackagedModelDirPath(PlatformName);
-
-	ConfigFile.AddUniqueToSection(TEXT("/Script/UnrealEd.ProjectPackagingSettings"), TEXT("+DirectoriesToAlwaysStageAsNonUFS"), FString("(Path=\"..") + StagingPath + FString("\")"));
-	ConfigFile.AddUniqueToSection(TEXT("Staging"), TEXT("+RemapDirectories"), FString("(From=\"") + FApp::GetProjectName() + StagingPath + FString("\", To=\"") + FApp::GetProjectName() + PackagingPath + FString("\")"));
-	ConfigFile.AddUniqueToSection(TEXT("Staging"), TEXT("+AllowedDirectories"), FApp::GetProjectName() + PackagingPath);
-}
-
 FString UNNERuntimeIREEGpu::GetRuntimeName() const
 {
 	return TEXT("");
diff --git a/Engine/Plugins/Experimental/NNERuntimeIREE/Source/NNERuntimeIREE/Private/NNERuntimeIREE.h b/Engine/Plugins/Experimental/NNERuntimeIREE/Source/NNERuntimeIREE/Private/NNERuntimeIREE.h
--- a/Engine/Plugins/Experimental/NNERuntimeIREE/Source/NNERuntimeIREE/Private/NNERuntimeIREE.h
+++ b/Engine/Plugins/Experimental/NNERuntimeIREE/Source/NNERuntimeIREE/Private/NNERuntimeIREE.h
@@ -33,8 +33,6 @@ public:
 	virtual ECanCreateModelCPUStatus CanCreateModelCPU(const TObjectPtr<UNNEModelData> ModelData) const override;
 	virtual TSharedPtr<UE::NNE::IModelCPU> CreateModelCPU(const TObjectPtr<UNNEModelData> ModelData) override;
 	//~ End INNERuntimeCPU Interface
-
-	static void GetUpdatedPlatformConfig(const FString& PlatformName, FConfigFile& ConfigFile, FString& ConfigFilePath);
 #else
 	//~ Begin INNERuntime Interface
 	virtual FString GetRuntimeName() const override { return ""; };
diff --git a/Engine/Plugins/Experimental/NNERuntimeIREE/Source/NNERuntimeIREE/Private/NNERuntimeIREEModule.cpp b/Engine/Plugins/Experimental/NNERuntimeIREE/Source/NNERuntimeIREE/Private/NNERuntimeIREEModule.cpp
--- a/Engine/Plugins/Experimental/NNERuntimeIREE/Source/NNERuntimeIREE/Private/NNERuntimeIREEModule.cpp
+++ b/Engine/Plugins/Experimental/NNERuntimeIREE/Source/NNERuntimeIREE/Private/NNERuntimeIREEModule.cpp
@@ -32,38 +32,6 @@ void FNNERuntimeIREEModule::StartupModule()
 		UE::NNE::RegisterRuntime(NNERuntimeIREECpu.Get());
 	}
 
-#if WITH_EDITOR
-	ITargetPlatformManagerModule* TargetPlatformManagerModule = GetTargetPlatformManager();
-	if (TargetPlatformManagerModule)
-	{
-		TSet<FString> ProcessedPlatforms;
-		TArray<ITargetPlatform*> TargetPlatforms = TargetPlatformManagerModule->GetTargetPlatforms();
-		for (int32 i = 0; i < TargetPlatforms.Num(); i++)
-		{
-			FString IniPlatformName = TargetPlatforms[i]->IniPlatformName();
-			if (!ProcessedPlatforms.Contains(IniPlatformName))
-			{
-				ProcessedPlatforms.Add(IniPlatformName);
-				FString TargetPlatformDisplayName = IniPlatformName;
-				FConfigFile ConfigFile;
-				FString ConfigFilePath;
-				NNERuntimeIREECpu->GetUpdatedPlatformConfig(IniPlatformName, ConfigFile, ConfigFilePath);
-				if (ConfigFile.Dirty)
-				{
-					{
-						TUniquePtr<FArchive> Ar = TUniquePtr<FArchive>(IFileManager::Get().CreateFileWriter(*ConfigFilePath, EFileWrite::FILEWRITE_Append));
-						if (!Ar)
-						{
-							continue;
-						}
-					}
-					ConfigFile.Write(ConfigFilePath);
-				}
-			}
-		}
-	}
-#endif // WITH_EDITOR
-
 	NNERuntimeIREECuda = NewObject<UNNERuntimeIREECuda>();
 	if (NNERuntimeIREECuda.IsValid())
 	{
