[NNE] NNERuntimeIREE remove DDC.

#tests Editor, Win64 load and run model in PIE

#rnx

diff --git a/Engine/Plugins/Experimental/NNERuntimeIREE/Source/NNERuntimeIREE/Private/NNERuntimeIREE.cpp b/Engine/Plugins/Experimental/NNERuntimeIREE/Source/NNERuntimeIREE/Private/NNERuntimeIREE.cpp
--- a/Engine/Plugins/Experimental/NNERuntimeIREE/Source/NNERuntimeIREE/Private/NNERuntimeIREE.cpp
+++ b/Engine/Plugins/Experimental/NNERuntimeIREE/Source/NNERuntimeIREE/Private/NNERuntimeIREE.cpp
@@ -6,10 +6,7 @@
 
 #if WITH_EDITOR
 #include "Containers/StringConv.h"
-#include "DerivedDataCache.h"
-#include "DerivedDataRequestOwner.h"
 #include "HAL/PlatformFileManager.h"
-#include "IO/IoHash.h"
 #include "Memory/SharedBuffer.h"
 #include "Misc/FileHelper.h"
 #endif // WITH_EDITOR
@@ -32,51 +29,6 @@
 
 namespace UE::NNERuntimeIREE::CPU::Private
 {
-#if WITH_EDITOR
-	inline UE::DerivedData::FCacheKey CreateCacheKey(const FString& RequestId)
-	{
-		return { UE::DerivedData::FCacheBucket(TEXT("NNEModelData")), FIoHash::HashBuffer(MakeMemoryView(StringCast<UTF8CHAR>(*RequestId))) };
-	}
-
-	inline void PutIntoDDC(const FString& RequestId, const FSharedBuffer& Data)
-	{
-		check(!Data.IsNull());
-		check(Data.GetSize() > 0);
-
-		TArray<UE::DerivedData::FCachePutValueRequest> Requests;
-		Requests.SetNum(1);
-
-		Requests[0].Name = FString("Put-") + RequestId;
-		Requests[0].Key = CreateCacheKey(RequestId);
-		Requests[0].Value = UE::DerivedData::FValue::Compress(Data);
-
-		UE::DerivedData::FRequestOwner BlockingPutOwner(UE::DerivedData::EPriority::Blocking);
-		UE::DerivedData::GetCache().PutValue(Requests, BlockingPutOwner);
-		BlockingPutOwner.Wait();
-	}
-
-	inline FSharedBuffer GetFromDDC(const FString& RequestId)
-	{
-		TArray<UE::DerivedData::FCacheGetValueRequest> Requests;
-		Requests.SetNum(1);
-
-		Requests[0].Name = FString("Get-") + RequestId;
-		Requests[0].Key = CreateCacheKey(RequestId);
-
-		FSharedBuffer Result;
-		UE::DerivedData::FRequestOwner BlockingGetOwner(UE::DerivedData::EPriority::Blocking);
-		UE::DerivedData::GetCache().GetValue(Requests, BlockingGetOwner, [&Result](UE::DerivedData::FCacheGetValueResponse&& Response)
-			{
-				if (Response.Value.HasData() && Response.Value.GetRawSize() > sizeof(uint32))
-				{
-					Result = Response.Value.GetData().Decompress();
-				}
-			});
-		BlockingGetOwner.Wait();
-		return Result;
-	}
-#endif // WITH_EDITOR
-
 	FString GetModelDataIdentifier(const FString& RuntimeName, const FGuid& Guid, const FString& FileIdString, const FString& PlatformName, const FString& Architecture)
 	{
 		return RuntimeName + "-" + Guid.ToString(EGuidFormats::Digits) + "-" + FString::FromInt(UNNERuntimeIREECpu::Version) + "-" + FileIdString + "-" + PlatformName + (!Architecture.IsEmpty() ? ("-" + Architecture) : "");
@@ -151,29 +103,6 @@ TSharedPtr<UE::NNE::FSharedModelData> UNNERuntimeIREECpu::CreateModelData(const
 		return TSharedPtr<UE::NNE::FSharedModelData>();
 	}
 
-	for (int32 i = 0; i < CompilerResults.Num(); i++)
-	{
-		TArray<uint8> SharedLibData;
-		FString StagedSharedLibPath = FPaths::Combine(StagingDir, CompilerResults[i].RelativeDirPath, CompilerResults[i].SharedLibraryFileName);
-		if (!FFileHelper::LoadFileToArray(SharedLibData, *StagedSharedLibPath) || SharedLibData.IsEmpty())
-		{
-			UE_LOG(LogNNERuntimeIREE, Warning, TEXT("UNNERuntimeIREECpu could not read the shared library \"%s\""), *StagedSharedLibPath);
-			return TSharedPtr<UE::NNE::FSharedModelData>();
-		}
-		FSharedBuffer SharedLibBuffer = MakeSharedBufferFromArray(MoveTemp(SharedLibData));
-		PutIntoDDC(UE::NNERuntimeIREE::CPU::Private::GetModelDataIdentifier(GetRuntimeName(), UNNERuntimeIREECpu::GUID, FileIdString, TargetPlatformName, CompilerResults[i].Architecture) + "-lib", SharedLibBuffer);
-
-		TArray<uint8> VmfbData;
-		FString StagedVmfbPath = FPaths::Combine(StagingDir, CompilerResults[i].RelativeDirPath, CompilerResults[i].VmfbFileName);
-		if (!FFileHelper::LoadFileToArray(VmfbData, *StagedVmfbPath) || VmfbData.IsEmpty())
-		{
-			UE_LOG(LogNNERuntimeIREE, Warning, TEXT("UNNERuntimeIREECpu could not read the vmfb data \"%s\""), *StagedVmfbPath);
-			return TSharedPtr<UE::NNE::FSharedModelData>();
-		}
-		FSharedBuffer VmfbBuffer = MakeSharedBufferFromArray(MoveTemp(VmfbData));
-		PutIntoDDC(UE::NNERuntimeIREE::CPU::Private::GetModelDataIdentifier(GetRuntimeName(), UNNERuntimeIREECpu::GUID, FileIdString, TargetPlatformName, CompilerResults[i].Architecture) + "-vmfb", VmfbBuffer);
-	}
-
 	TArray<uint8> ResultData;
 	FMemoryWriter Writer(ResultData);
 	Writer << UNNERuntimeIREECpu::GUID;
@@ -341,35 +270,6 @@ TSharedPtr<UE::NNE::IModelCPU> UNNERuntimeIREECpu::CreateModelCPU(const TObjectP
 	FString SharedLibraryDirPath = FPaths::ConvertRelativePathToFull(FPaths::Combine(FPaths::ProjectDir(), GetPackagedModelDirPath(UGameplayStatics::GetPlatformName()), RelativeDirPath));
 #endif // WITH_EDITOR
 
-#if WITH_EDITOR
-	IPlatformFile& PlatformFile = FPlatformFileManager::Get().GetPlatformFile();
-	FString SharedLibraryFilePath = FPaths::Combine(SharedLibraryDirPath, SharedLibraryFileName);
-	if (!PlatformFile.FileExists(*SharedLibraryFilePath))
-	{
-		FSharedBuffer SharedBuffer = GetFromDDC(UE::NNERuntimeIREE::CPU::Private::GetModelDataIdentifier(GetRuntimeName(), UNNERuntimeIREECpu::GUID, FileIdString, UGameplayStatics::GetPlatformName(), Architecture) + "-lib");
-		if (SharedBuffer.GetSize() <= 0)
-		{
-			UE_LOG(LogNNERuntimeIREE, Warning, TEXT("UNNERuntimeIREECpu could not fetch the shared library %s from DDC"), *SharedLibraryFileName);
-			return TSharedPtr<UE::NNE::IModelCPU>();
-		}
-		PlatformFile.CreateDirectoryTree(*SharedLibraryDirPath);
-		FFileHelper::SaveArrayToFile(TConstArrayView<uint8>((uint8*)SharedBuffer.GetData(), SharedBuffer.GetSize()), *SharedLibraryFilePath);
-	}
-
-	FString VmfbFilePath = FPaths::Combine(SharedLibraryDirPath, VmfbFileName);
-	if (!PlatformFile.FileExists(*VmfbFilePath))
-	{
-		FSharedBuffer SharedBuffer = GetFromDDC(UE::NNERuntimeIREE::CPU::Private::GetModelDataIdentifier(GetRuntimeName(), UNNERuntimeIREECpu::GUID, FileIdString, UGameplayStatics::GetPlatformName(), Architecture) + "-vmfb");
-		if (SharedBuffer.GetSize() <= 0)
-		{
-			UE_LOG(LogNNERuntimeIREE, Warning, TEXT("UNNERuntimeIREECpu could not fetch the vmfb %s from DDC"), *VmfbFileName);
-			return TSharedPtr<UE::NNE::IModelCPU>();
-		}
-		PlatformFile.CreateDirectoryTree(*SharedLibraryDirPath);
-		FFileHelper::SaveArrayToFile(TConstArrayView<uint8>((uint8*)SharedBuffer.GetData(), SharedBuffer.GetSize()), *VmfbFilePath);
-	}
-#endif // WITH_EDITOR
-
 	TSharedPtr<UE::NNE::IModelCPU> Model = UE::NNERuntimeIREE::CPU::FModel::Make(SharedLibraryDirPath, SharedLibraryFileName, VmfbFileName, SharedLibraryEntryPointName, *ModuleMetaData);
 	if (!Model.IsValid())
 	{
