[NNE] Support for >2GB models at the API level.

Problem:
* UNNEModelData filedata and additioal data were limited to 2GB of data (via TArray<uint8>).
* SharedModelData.GetView() was limited to 2GB (via TArray<uint8>).

Solution:
* Change UNNEModelData and SharedModelData to use TArray64<uint8>
* Handle backward compatibility of UNNEModelData at the serialization level

Test:
* Win64 Editor debug : all tests - both with old and new assets format (reimported) + style transfer
* Linux Editor dev : all tests - both with old and new assets format (reimported)
* Win64 Editor debug & dev : large model (> 2GB) on ORT Cpu, Gpu and RDG runtimes
* Win64 package debug standalone: style transfer + smoke tests (new formats)
* Console package debug standalone : style transfer + smoke tests (new formats)

#rnx

#virtualized

diff --git a/Engine/Plugins/Experimental/NNERuntimeBasicCpu/Source/NNERuntimeBasicCpu/Private/NNERuntimeBasicCpu.cpp b/Engine/Plugins/Experimental/NNERuntimeBasicCpu/Source/NNERuntimeBasicCpu/Private/NNERuntimeBasicCpu.cpp
--- a/Engine/Plugins/Experimental/NNERuntimeBasicCpu/Source/NNERuntimeBasicCpu/Private/NNERuntimeBasicCpu.cpp
+++ b/Engine/Plugins/Experimental/NNERuntimeBasicCpu/Source/NNERuntimeBasicCpu/Private/NNERuntimeBasicCpu.cpp
@@ -11,7 +11,7 @@
 // various bits of data.
 const uint32 UNNERuntimeBasicCpuImpl::Alignment = 64;
 
-UNNERuntimeBasicCpuImpl::ECanCreateModelDataStatus UNNERuntimeBasicCpuImpl::CanCreateModelData(const FString& FileType, TConstArrayView<uint8> FileData, const TMap<FString, TConstArrayView<uint8>>& AdditionalFileData, const FGuid& FileId, const ITargetPlatform* TargetPlatform) const
+UNNERuntimeBasicCpuImpl::ECanCreateModelDataStatus UNNERuntimeBasicCpuImpl::CanCreateModelData(const FString& FileType, TConstArrayView64<uint8> FileData, const TMap<FString, TConstArrayView64<uint8>>& AdditionalFileData, const FGuid& FileId, const ITargetPlatform* TargetPlatform) const
 {
 	if (FileType.Compare("ubnne", ESearchCase::IgnoreCase) != 0)
 	{
@@ -38,10 +38,22 @@ UNNERuntimeBasicCpuImpl::ECanCreateModelDataStatus UNNERuntimeBasicCpuImpl::CanC
 		return ECanCreateModelDataStatus::Fail;
 	}
 
+	// Check model do not have additional data
+	if (!AdditionalFileData.IsEmpty())
+	{
+		return ECanCreateModelDataStatus::Fail;
+	}
+
+	// Check model is not > 2GB
+	if ((TArray<uint8>::SizeType)FileData.Num() != FileData.Num())
+	{
+		return ECanCreateModelDataStatus::Fail;
+	}
+
 	return ECanCreateModelDataStatus::Ok;
 }
 
-TSharedPtr<UE::NNE::FSharedModelData> UNNERuntimeBasicCpuImpl::CreateModelData(const FString& FileType, TConstArrayView<uint8> FileData, const TMap<FString, TConstArrayView<uint8>>& AdditionalFileData, const FGuid& FileId, const ITargetPlatform* TargetPlatform)
+TSharedPtr<UE::NNE::FSharedModelData> UNNERuntimeBasicCpuImpl::CreateModelData(const FString& FileType, TConstArrayView64<uint8> FileData, const TMap<FString, TConstArrayView64<uint8>>& AdditionalFileData, const FGuid& FileId, const ITargetPlatform* TargetPlatform)
 {
 	if (CanCreateModelData(FileType, FileData, AdditionalFileData, FileId, TargetPlatform) != ECanCreateModelDataStatus::Ok)
 	{
@@ -59,7 +71,7 @@ TSharedPtr<UE::NNE::FSharedModelData> UNNERuntimeBasicCpuImpl::CreateModelData(c
 	return MakeShared<UE::NNE::FSharedModelData>(FSharedBuffer::TakeOwnership(ModelData, FileData.Num(), FMemory::Free), Alignment);
 }
 
-FString UNNERuntimeBasicCpuImpl::GetModelDataIdentifier(const FString& FileType, TConstArrayView<uint8> FileData, const TMap<FString, TConstArrayView<uint8>>& AdditionalFileData, const FGuid& FileId, const ITargetPlatform* TargetPlatform) const
+FString UNNERuntimeBasicCpuImpl::GetModelDataIdentifier(const FString& FileType, TConstArrayView64<uint8> FileData, const TMap<FString, TConstArrayView64<uint8>>& AdditionalFileData, const FGuid& FileId, const ITargetPlatform* TargetPlatform) const
 {
 	return FileId.ToString(EGuidFormats::Digits) + "-" + FString::FromInt(UE::NNE::RuntimeBasic::FModelCPU::ModelMagicNumber);
 }
@@ -75,7 +87,7 @@ UNNERuntimeBasicCpuImpl::ECanCreateModelCPUStatus UNNERuntimeBasicCpuImpl::CanCr
 		return ECanCreateModelCPUStatus::Fail;
 	}
 
-	TConstArrayView<uint8> Data = SharedData->GetView();
+	TConstArrayView64<uint8> Data = SharedData->GetView();
 
 	// We require at least a magic number and version number
 	if (Data.Num() < 2 * sizeof(uint32))
diff --git a/Engine/Plugins/Experimental/NNERuntimeBasicCpu/Source/NNERuntimeBasicCpu/Private/NNERuntimeBasicCpu.h b/Engine/Plugins/Experimental/NNERuntimeBasicCpu/Source/NNERuntimeBasicCpu/Private/NNERuntimeBasicCpu.h
--- a/Engine/Plugins/Experimental/NNERuntimeBasicCpu/Source/NNERuntimeBasicCpu/Private/NNERuntimeBasicCpu.h
+++ b/Engine/Plugins/Experimental/NNERuntimeBasicCpu/Source/NNERuntimeBasicCpu/Private/NNERuntimeBasicCpu.h
@@ -25,9 +25,9 @@ public:
 
 	virtual FString GetRuntimeName() const override { return TEXT("NNERuntimeBasicCpu"); };
 
-	virtual ECanCreateModelDataStatus CanCreateModelData(const FString& FileType, TConstArrayView<uint8> FileData, const TMap<FString, TConstArrayView<uint8>>& AdditionalFileData, const FGuid& FileId, const ITargetPlatform* TargetPlatform) const override;
-	virtual TSharedPtr<UE::NNE::FSharedModelData> CreateModelData(const FString& FileType, TConstArrayView<uint8> FileData, const TMap<FString, TConstArrayView<uint8>>& AdditionalFileData, const FGuid& FileId, const ITargetPlatform* TargetPlatform) override;
-	virtual FString GetModelDataIdentifier(const FString& FileType, TConstArrayView<uint8> FileData, const TMap<FString, TConstArrayView<uint8>>& AdditionalFileData, const FGuid& FileId, const ITargetPlatform* TargetPlatform) const override;
+	virtual ECanCreateModelDataStatus CanCreateModelData(const FString& FileType, TConstArrayView64<uint8> FileData, const TMap<FString, TConstArrayView64<uint8>>& AdditionalFileData, const FGuid& FileId, const ITargetPlatform* TargetPlatform) const override;
+	virtual TSharedPtr<UE::NNE::FSharedModelData> CreateModelData(const FString& FileType, TConstArrayView64<uint8> FileData, const TMap<FString, TConstArrayView64<uint8>>& AdditionalFileData, const FGuid& FileId, const ITargetPlatform* TargetPlatform) override;
+	virtual FString GetModelDataIdentifier(const FString& FileType, TConstArrayView64<uint8> FileData, const TMap<FString, TConstArrayView64<uint8>>& AdditionalFileData, const FGuid& FileId, const ITargetPlatform* TargetPlatform) const override;
 
 	virtual ECanCreateModelCPUStatus CanCreateModelCPU(const TObjectPtr<UNNEModelData> ModelData) const override;
 	virtual TSharedPtr<UE::NNE::IModelCPU> CreateModelCPU(const TObjectPtr<UNNEModelData> ModelData) override;
diff --git a/Engine/Plugins/Experimental/NNERuntimeIREE/Source/NNERuntimeIREE/Private/NNERuntimeIREE.cpp b/Engine/Plugins/Experimental/NNERuntimeIREE/Source/NNERuntimeIREE/Private/NNERuntimeIREE.cpp
--- a/Engine/Plugins/Experimental/NNERuntimeIREE/Source/NNERuntimeIREE/Private/NNERuntimeIREE.cpp
+++ b/Engine/Plugins/Experimental/NNERuntimeIREE/Source/NNERuntimeIREE/Private/NNERuntimeIREE.cpp
@@ -53,23 +53,29 @@ namespace UE::NNERuntimeIREE::CPU::Private
 } // UE::NNERuntimeIREE::CPU::Private
 
 FGuid UNNERuntimeIREECpu::GUID = FGuid((int32)'I', (int32)'C', (int32)'P', (int32)'U');
-int32 UNNERuntimeIREECpu::Version = 0x00000003;
+int32 UNNERuntimeIREECpu::Version = 0x00000004;
 
 FString UNNERuntimeIREECpu::GetRuntimeName() const
 {
 	return TEXT("NNERuntimeIREECpu");
 }
 
-UNNERuntimeIREECpu::ECanCreateModelDataStatus UNNERuntimeIREECpu::CanCreateModelData(const FString& FileType, TConstArrayView<uint8> FileData, const TMap<FString, TConstArrayView<uint8>>& AdditionalFileData, const FGuid& FileId, const ITargetPlatform* TargetPlatform) const
+UNNERuntimeIREECpu::ECanCreateModelDataStatus UNNERuntimeIREECpu::CanCreateModelData(const FString& FileType, TConstArrayView64<uint8> FileData, const TMap<FString, TConstArrayView64<uint8>>& AdditionalFileData, const FGuid& FileId, const ITargetPlatform* TargetPlatform) const
 {
 #if WITH_EDITOR
+	// Check model is not > 2GB
+	if ((TArray<uint8>::SizeType)FileData.Num() != FileData.Num())
+	{
+		return ECanCreateModelDataStatus::Fail;
+	}
+
 	return FileType.Compare(TEXT("mlir"), ESearchCase::IgnoreCase) == 0 ? ECanCreateModelDataStatus::Ok : ECanCreateModelDataStatus::FailFileIdNotSupported;
 #else
 	return ECanCreateModelDataStatus::Fail;
 #endif // WITH_EDITOR
 }
 
-TSharedPtr<UE::NNE::FSharedModelData> UNNERuntimeIREECpu::CreateModelData(const FString& FileType, TConstArrayView<uint8> FileData, const TMap<FString, TConstArrayView<uint8>>& AdditionalFileData, const FGuid& FileId, const ITargetPlatform* TargetPlatform)
+TSharedPtr<UE::NNE::FSharedModelData> UNNERuntimeIREECpu::CreateModelData(const FString& FileType, TConstArrayView64<uint8> FileData, const TMap<FString, TConstArrayView64<uint8>>& AdditionalFileData, const FGuid& FileId, const ITargetPlatform* TargetPlatform)
 {
 #if WITH_EDITOR
 	using namespace UE::NNERuntimeIREE::CPU::Private;
@@ -103,21 +109,21 @@ TSharedPtr<UE::NNE::FSharedModelData> UNNERuntimeIREECpu::CreateModelData(const
 		return TSharedPtr<UE::NNE::FSharedModelData>();
 	}
 
-	TArray<uint8> ResultData;
-	FMemoryWriter Writer(ResultData);
+	TArray64<uint8> ResultData;
+	FMemoryWriter64 Writer(ResultData);
 	Writer << UNNERuntimeIREECpu::GUID;
 	Writer << UNNERuntimeIREECpu::Version;
 	FGuid FileIdCopy = FileId;
 	Writer << FileIdCopy;
 
-	TArray<uint8> ModuleMetaData;
+	TArray64<uint8> ModuleMetaData;
 	if (AdditionalFileData.Contains("IREEModuleMetaData"))
 	{
 		ModuleMetaData = AdditionalFileData["IREEModuleMetaData"];
 	}
 	if (ModuleMetaData.IsEmpty())
 	{
-		FMemoryWriter ObjectWriter(ModuleMetaData);
+		FMemoryWriter64 ObjectWriter(ModuleMetaData);
 		CompilerModuleMetaData->Serialize(ObjectWriter);
 	}
 	Writer << ModuleMetaData;
@@ -139,7 +145,7 @@ TSharedPtr<UE::NNE::FSharedModelData> UNNERuntimeIREECpu::CreateModelData(const
 #endif // WITH_EDITOR
 }
 
-FString UNNERuntimeIREECpu::GetModelDataIdentifier(const FString& FileType, TConstArrayView<uint8> FileData, const TMap<FString, TConstArrayView<uint8>>& AdditionalFileData, const FGuid& FileId, const ITargetPlatform* TargetPlatform) const
+FString UNNERuntimeIREECpu::GetModelDataIdentifier(const FString& FileType, TConstArrayView64<uint8> FileData, const TMap<FString, TConstArrayView64<uint8>>& AdditionalFileData, const FGuid& FileId, const ITargetPlatform* TargetPlatform) const
 {
 	// Leave architecture blank as there is only one model data for all architectures of a given platform, only the vmfb and shared lib are different
 	FString PlatformName = TargetPlatform ? TargetPlatform->IniPlatformName() : UGameplayStatics::GetPlatformName();
@@ -295,21 +301,27 @@ FString UNNERuntimeIREEGpu::GetRuntimeName() const
 	return TEXT("");
 }
 
-UNNERuntimeIREEGpu::ECanCreateModelDataStatus UNNERuntimeIREEGpu::CanCreateModelData(const FString& FileType, TConstArrayView<uint8> FileData, const TMap<FString, TConstArrayView<uint8>>& AdditionalFileData, const FGuid& FileId, const ITargetPlatform* TargetPlatform) const
+UNNERuntimeIREEGpu::ECanCreateModelDataStatus UNNERuntimeIREEGpu::CanCreateModelData(const FString& FileType, TConstArrayView64<uint8> FileData, const TMap<FString, TConstArrayView64<uint8>>& AdditionalFileData, const FGuid& FileId, const ITargetPlatform* TargetPlatform) const
 {
 #if WITH_EDITOR
+	// Check model is not > 2GB
+	if ((TArray<uint8>::SizeType)FileData.Num() != FileData.Num())
+	{
+		return ECanCreateModelDataStatus::Fail;
+	}
+
 	return FileType.Compare(TEXT("mlir"), ESearchCase::IgnoreCase) == 0 ? ECanCreateModelDataStatus::Ok : ECanCreateModelDataStatus::FailFileIdNotSupported;
 #else
 	return ECanCreateModelDataStatus::Fail;
 #endif // WITH_EDITOR
 }
 
-TSharedPtr<UE::NNE::FSharedModelData> UNNERuntimeIREEGpu::CreateModelData(const FString& FileType, TConstArrayView<uint8> FileData, const TMap<FString, TConstArrayView<uint8>>& AdditionalFileData, const FGuid& FileId, const ITargetPlatform* TargetPlatform)
+TSharedPtr<UE::NNE::FSharedModelData> UNNERuntimeIREEGpu::CreateModelData(const FString& FileType, TConstArrayView64<uint8> FileData, const TMap<FString, TConstArrayView64<uint8>>& AdditionalFileData, const FGuid& FileId, const ITargetPlatform* TargetPlatform)
 {
 	return TSharedPtr<UE::NNE::FSharedModelData>();
 }
 
-FString UNNERuntimeIREEGpu::GetModelDataIdentifier(const FString& FileType, TConstArrayView<uint8> FileData, const TMap<FString, TConstArrayView<uint8>>& AdditionalFileData, const FGuid& FileId, const ITargetPlatform* TargetPlatform) const
+FString UNNERuntimeIREEGpu::GetModelDataIdentifier(const FString& FileType, TConstArrayView64<uint8> FileData, const TMap<FString, TConstArrayView64<uint8>>& AdditionalFileData, const FGuid& FileId, const ITargetPlatform* TargetPlatform) const
 {
 	FString PlatformName = TargetPlatform ? TargetPlatform->IniPlatformName() : UGameplayStatics::GetPlatformName();
 	return UE::NNERuntimeIREE::CPU::Private::GetModelDataIdentifier(GetRuntimeName(), GetGUID(), FileId.ToString(EGuidFormats::Digits), PlatformName, "");
@@ -438,21 +450,27 @@ FString UNNERuntimeIREERdg::GetRuntimeName() const
 	return TEXT("NNERuntimeIREERdg");
 }
 
-UNNERuntimeIREERdg::ECanCreateModelDataStatus UNNERuntimeIREERdg::CanCreateModelData(const FString& FileType, TConstArrayView<uint8> FileData, const TMap<FString, TConstArrayView<uint8>>& AdditionalFileData, const FGuid& FileId, const ITargetPlatform* TargetPlatform) const
+UNNERuntimeIREERdg::ECanCreateModelDataStatus UNNERuntimeIREERdg::CanCreateModelData(const FString& FileType, TConstArrayView64<uint8> FileData, const TMap<FString, TConstArrayView64<uint8>>& AdditionalFileData, const FGuid& FileId, const ITargetPlatform* TargetPlatform) const
 {
 #if WITH_EDITOR
+	// Check model is not > 2GB
+	if ((TArray<uint8>::SizeType)FileData.Num() != FileData.Num())
+	{
+		return ECanCreateModelDataStatus::Fail;
+	}
+
 	return	FileType.Compare(TEXT("mlir"), ESearchCase::IgnoreCase) == 0 ? ECanCreateModelDataStatus::Ok : ECanCreateModelDataStatus::FailFileIdNotSupported;
 #else
 	return ECanCreateModelDataStatus::Fail;
 #endif // WITH_EDITOR
 }
 
-TSharedPtr<UE::NNE::FSharedModelData> UNNERuntimeIREERdg::CreateModelData(const FString& FileType, TConstArrayView<uint8> FileData, const TMap<FString, TConstArrayView<uint8>>& AdditionalFileData, const FGuid& FileId, const ITargetPlatform* TargetPlatform)
+TSharedPtr<UE::NNE::FSharedModelData> UNNERuntimeIREERdg::CreateModelData(const FString& FileType, TConstArrayView64<uint8> FileData, const TMap<FString, TConstArrayView64<uint8>>& AdditionalFileData, const FGuid& FileId, const ITargetPlatform* TargetPlatform)
 {
 	return TSharedPtr<UE::NNE::FSharedModelData>();
 }
 
-FString UNNERuntimeIREERdg::GetModelDataIdentifier(const FString& FileType, TConstArrayView<uint8> FileData, const TMap<FString, TConstArrayView<uint8>>& AdditionalFileData, const FGuid& FileId, const ITargetPlatform* TargetPlatform) const
+FString UNNERuntimeIREERdg::GetModelDataIdentifier(const FString& FileType, TConstArrayView64<uint8> FileData, const TMap<FString, TConstArrayView64<uint8>>& AdditionalFileData, const FGuid& FileId, const ITargetPlatform* TargetPlatform) const
 {
 	FString PlatformName = TargetPlatform ? TargetPlatform->IniPlatformName() : UGameplayStatics::GetPlatformName();
 	return UE::NNERuntimeIREE::CPU::Private::GetModelDataIdentifier(GetRuntimeName(), UNNERuntimeIREERdg::GUID, FileId.ToString(EGuidFormats::Digits), PlatformName, "");
diff --git a/Engine/Plugins/Experimental/NNERuntimeIREE/Source/NNERuntimeIREE/Private/NNERuntimeIREE.h b/Engine/Plugins/Experimental/NNERuntimeIREE/Source/NNERuntimeIREE/Private/NNERuntimeIREE.h
--- a/Engine/Plugins/Experimental/NNERuntimeIREE/Source/NNERuntimeIREE/Private/NNERuntimeIREE.h
+++ b/Engine/Plugins/Experimental/NNERuntimeIREE/Source/NNERuntimeIREE/Private/NNERuntimeIREE.h
@@ -24,9 +24,9 @@ public:
 
 	//~ Begin INNERuntime Interface
 	virtual FString GetRuntimeName() const override;
-	virtual ECanCreateModelDataStatus CanCreateModelData(const FString& FileType, TConstArrayView<uint8> FileData, const TMap<FString, TConstArrayView<uint8>>& AdditionalFileData, const FGuid& FileId, const ITargetPlatform* TargetPlatform) const override;
-	virtual TSharedPtr<UE::NNE::FSharedModelData> CreateModelData(const FString& FileType, TConstArrayView<uint8> FileData, const TMap<FString, TConstArrayView<uint8>>& AdditionalFileData, const FGuid& FileId, const ITargetPlatform* TargetPlatform) override;
-	virtual FString GetModelDataIdentifier(const FString& FileType, TConstArrayView<uint8> FileData, const TMap<FString, TConstArrayView<uint8>>& AdditionalFileData, const FGuid& FileId, const ITargetPlatform* TargetPlatform) const override;
+	virtual ECanCreateModelDataStatus CanCreateModelData(const FString& FileType, TConstArrayView64<uint8> FileData, const TMap<FString, TConstArrayView64<uint8>>& AdditionalFileData, const FGuid& FileId, const ITargetPlatform* TargetPlatform) const override;
+	virtual TSharedPtr<UE::NNE::FSharedModelData> CreateModelData(const FString& FileType, TConstArrayView64<uint8> FileData, const TMap<FString, TConstArrayView64<uint8>>& AdditionalFileData, const FGuid& FileId, const ITargetPlatform* TargetPlatform) override;
+	virtual FString GetModelDataIdentifier(const FString& FileType, TConstArrayView64<uint8> FileData, const TMap<FString, TConstArrayView64<uint8>>& AdditionalFileData, const FGuid& FileId, const ITargetPlatform* TargetPlatform) const override;
 	//~ End INNERuntime Interface
 
 	//~ Begin INNERuntimeCPU Interface
@@ -36,9 +36,9 @@ public:
 #else
 	//~ Begin INNERuntime Interface
 	virtual FString GetRuntimeName() const override { return ""; };
-	virtual ECanCreateModelDataStatus CanCreateModelData(const FString& FileType, TConstArrayView<uint8> FileData, const TMap<FString, TConstArrayView<uint8>>& AdditionalFileData, const FGuid& FileId, const ITargetPlatform* TargetPlatform) const override { return ECanCreateModelDataStatus::Fail; };
-	virtual TSharedPtr<UE::NNE::FSharedModelData> CreateModelData(const FString& FileType, TConstArrayView<uint8> FileData, const TMap<FString, TConstArrayView<uint8>>& AdditionalFileData, const FGuid& FileId, const ITargetPlatform* TargetPlatform) override { return TSharedPtr<UE::NNE::FSharedModelData>(); };
-	virtual FString GetModelDataIdentifier(const FString& FileType, TConstArrayView<uint8> FileData, const TMap<FString, TConstArrayView<uint8>>& AdditionalFileData, const FGuid& FileId, const ITargetPlatform* TargetPlatform) const override { return ""; };
+	virtual ECanCreateModelDataStatus CanCreateModelData(const FString& FileType, TConstArrayView64<uint8> FileData, const TMap<FString, TConstArrayView64<uint8>>& AdditionalFileData, const FGuid& FileId, const ITargetPlatform* TargetPlatform) const override { return ECanCreateModelDataStatus::Fail; };
+	virtual TSharedPtr<UE::NNE::FSharedModelData> CreateModelData(const FString& FileType, TConstArrayView64<uint8> FileData, const TMap<FString, TConstArrayView64<uint8>>& AdditionalFileData, const FGuid& FileId, const ITargetPlatform* TargetPlatform) override { return TSharedPtr<UE::NNE::FSharedModelData>(); };
+	virtual FString GetModelDataIdentifier(const FString& FileType, TConstArrayView64<uint8> FileData, const TMap<FString, TConstArrayView64<uint8>>& AdditionalFileData, const FGuid& FileId, const ITargetPlatform* TargetPlatform) const override { return ""; };
 	//~ End INNERuntime Interface
 
 	//~ Begin INNERuntimeCPU Interface
@@ -57,9 +57,9 @@ public:
 #ifdef WITH_NNE_RUNTIME_IREE
 	//~ Begin INNERuntime Interface
 	virtual FString GetRuntimeName() const override;
-	virtual ECanCreateModelDataStatus CanCreateModelData(const FString& FileType, TConstArrayView<uint8> FileData, const TMap<FString, TConstArrayView<uint8>>& AdditionalFileData, const FGuid& FileId, const ITargetPlatform* TargetPlatform) const override;
-	virtual TSharedPtr<UE::NNE::FSharedModelData> CreateModelData(const FString& FileType, TConstArrayView<uint8> FileData, const TMap<FString, TConstArrayView<uint8>>& AdditionalFileData, const FGuid& FileId, const ITargetPlatform* TargetPlatform) override;
-	virtual FString GetModelDataIdentifier(const FString& FileType, TConstArrayView<uint8> FileData, const TMap<FString, TConstArrayView<uint8>>& AdditionalFileData, const FGuid& FileId, const ITargetPlatform* TargetPlatform) const override;
+	virtual ECanCreateModelDataStatus CanCreateModelData(const FString& FileType, TConstArrayView64<uint8> FileData, const TMap<FString, TConstArrayView64<uint8>>& AdditionalFileData, const FGuid& FileId, const ITargetPlatform* TargetPlatform) const override;
+	virtual TSharedPtr<UE::NNE::FSharedModelData> CreateModelData(const FString& FileType, TConstArrayView64<uint8> FileData, const TMap<FString, TConstArrayView64<uint8>>& AdditionalFileData, const FGuid& FileId, const ITargetPlatform* TargetPlatform) override;
+	virtual FString GetModelDataIdentifier(const FString& FileType, TConstArrayView64<uint8> FileData, const TMap<FString, TConstArrayView64<uint8>>& AdditionalFileData, const FGuid& FileId, const ITargetPlatform* TargetPlatform) const override;
 	//~ End INNERuntime Interface
 
 	//~ Begin INNERuntimeGPU Interface
@@ -73,9 +73,9 @@ public:
 #else
 	//~ Begin INNERuntime Interface
 	virtual FString GetRuntimeName() const override { return ""; };
-	virtual ECanCreateModelDataStatus CanCreateModelData(const FString& FileType, TConstArrayView<uint8> FileData, const TMap<FString, TConstArrayView<uint8>>& AdditionalFileData, const FGuid& FileId, const ITargetPlatform* TargetPlatform) const override { return ECanCreateModelDataStatus::Fail; };
-	virtual TSharedPtr<UE::NNE::FSharedModelData> CreateModelData(const FString& FileType, TConstArrayView<uint8> FileData, const TMap<FString, TConstArrayView<uint8>>& AdditionalFileData, const FGuid& FileId, const ITargetPlatform* TargetPlatform) override { return TSharedPtr<UE::NNE::FSharedModelData>(); };
-	virtual FString GetModelDataIdentifier(const FString& FileType, TConstArrayView<uint8> FileData, const TMap<FString, TConstArrayView<uint8>>& AdditionalFileData, const FGuid& FileId, const ITargetPlatform* TargetPlatform) const override { return ""; };
+	virtual ECanCreateModelDataStatus CanCreateModelData(const FString& FileType, TConstArrayView64<uint8> FileData, const TMap<FString, TConstArrayView64<uint8>>& AdditionalFileData, const FGuid& FileId, const ITargetPlatform* TargetPlatform) const override { return ECanCreateModelDataStatus::Fail; };
+	virtual TSharedPtr<UE::NNE::FSharedModelData> CreateModelData(const FString& FileType, TConstArrayView64<uint8> FileData, const TMap<FString, TConstArrayView64<uint8>>& AdditionalFileData, const FGuid& FileId, const ITargetPlatform* TargetPlatform) override { return TSharedPtr<UE::NNE::FSharedModelData>(); };
+	virtual FString GetModelDataIdentifier(const FString& FileType, TConstArrayView64<uint8> FileData, const TMap<FString, TConstArrayView64<uint8>>& AdditionalFileData, const FGuid& FileId, const ITargetPlatform* TargetPlatform) const override { return ""; };
 	//~ End INNERuntime Interface
 
 	//~ Begin INNERuntimeGPU Interface
@@ -141,9 +141,9 @@ public:
 
 	//~ Begin INNERuntime Interface
 	virtual FString GetRuntimeName() const override;
-	virtual ECanCreateModelDataStatus CanCreateModelData(const FString& FileType, TConstArrayView<uint8> FileData, const TMap<FString, TConstArrayView<uint8>>& AdditionalFileData, const FGuid& FileId, const ITargetPlatform* TargetPlatform) const override;
-	virtual TSharedPtr<UE::NNE::FSharedModelData> CreateModelData(const FString& FileType, TConstArrayView<uint8> FileData, const TMap<FString, TConstArrayView<uint8>>& AdditionalFileData, const FGuid& FileId, const ITargetPlatform* TargetPlatform) override;
-	virtual FString GetModelDataIdentifier(const FString& FileType, TConstArrayView<uint8> FileData, const TMap<FString, TConstArrayView<uint8>>& AdditionalFileData, const FGuid& FileId, const ITargetPlatform* TargetPlatform) const override;
+	virtual ECanCreateModelDataStatus CanCreateModelData(const FString& FileType, TConstArrayView64<uint8> FileData, const TMap<FString, TConstArrayView64<uint8>>& AdditionalFileData, const FGuid& FileId, const ITargetPlatform* TargetPlatform) const override;
+	virtual TSharedPtr<UE::NNE::FSharedModelData> CreateModelData(const FString& FileType, TConstArrayView64<uint8> FileData, const TMap<FString, TConstArrayView64<uint8>>& AdditionalFileData, const FGuid& FileId, const ITargetPlatform* TargetPlatform) override;
+	virtual FString GetModelDataIdentifier(const FString& FileType, TConstArrayView64<uint8> FileData, const TMap<FString, TConstArrayView64<uint8>>& AdditionalFileData, const FGuid& FileId, const ITargetPlatform* TargetPlatform) const override;
 	//~ End INNERuntime Interface
 
 	//~ Begin INNERuntimeRdg Interface
@@ -155,9 +155,9 @@ public:
 #else
 	//~ Begin INNERuntime Interface
 	virtual FString GetRuntimeName() const override { return ""; };
-	virtual ECanCreateModelDataStatus CanCreateModelData(const FString& FileType, TConstArrayView<uint8> FileData, const TMap<FString, TConstArrayView<uint8>>& AdditionalFileData, const FGuid& FileId, const ITargetPlatform* TargetPlatform) const override { return ECanCreateModelDataStatus::Fail; };
-	virtual TSharedPtr<UE::NNE::FSharedModelData> CreateModelData(const FString& FileType, TConstArrayView<uint8> FileData, const TMap<FString, TConstArrayView<uint8>>& AdditionalFileData, const FGuid& FileId, const ITargetPlatform* TargetPlatform) override { return TSharedPtr<UE::NNE::FSharedModelData>(); };
-	virtual FString GetModelDataIdentifier(const FString& FileType, TConstArrayView<uint8> FileData, const TMap<FString, TConstArrayView<uint8>>& AdditionalFileData, const FGuid& FileId, const ITargetPlatform* TargetPlatform) const override { return ""; };
+	virtual ECanCreateModelDataStatus CanCreateModelData(const FString& FileType, TConstArrayView64<uint8> FileData, const TMap<FString, TConstArrayView64<uint8>>& AdditionalFileData, const FGuid& FileId, const ITargetPlatform* TargetPlatform) const override { return ECanCreateModelDataStatus::Fail; };
+	virtual TSharedPtr<UE::NNE::FSharedModelData> CreateModelData(const FString& FileType, TConstArrayView64<uint8> FileData, const TMap<FString, TConstArrayView64<uint8>>& AdditionalFileData, const FGuid& FileId, const ITargetPlatform* TargetPlatform) override { return TSharedPtr<UE::NNE::FSharedModelData>(); };
+	virtual FString GetModelDataIdentifier(const FString& FileType, TConstArrayView64<uint8> FileData, const TMap<FString, TConstArrayView64<uint8>>& AdditionalFileData, const FGuid& FileId, const ITargetPlatform* TargetPlatform) const override { return ""; };
 	//~ End INNERuntime Interface
 
 	//~ Begin INNERuntimeRDG Interface
diff --git a/Engine/Plugins/Experimental/NNERuntimeIREE/Source/NNERuntimeIREEEditor/Private/NNERuntimeIREEModelDataFactory.cpp b/Engine/Plugins/Experimental/NNERuntimeIREE/Source/NNERuntimeIREEEditor/Private/NNERuntimeIREEModelDataFactory.cpp
--- a/Engine/Plugins/Experimental/NNERuntimeIREE/Source/NNERuntimeIREEEditor/Private/NNERuntimeIREEModelDataFactory.cpp
+++ b/Engine/Plugins/Experimental/NNERuntimeIREE/Source/NNERuntimeIREEEditor/Private/NNERuntimeIREEModelDataFactory.cpp
@@ -208,7 +208,7 @@ UObject* UNNERuntimeIREEModelDataFactory::FactoryCreateBinary(UClass* Class, UOb
 	FMemoryWriter Writer(MetaDataByteArray);
 	ModuleMetaData->Serialize(Writer);
 
-	TMap<FString, TConstArrayView<uint8>> AdditionalFileData;
+	TMap<FString, TConstArrayView64<uint8>> AdditionalFileData;
 	AdditionalFileData.Add("IREEModuleMetaData", MetaDataByteArray);
 
 	UNNEModelData* ModelData = NewObject<UNNEModelData>(InParent, Class, Name, Flags);
diff --git a/Engine/Plugins/Experimental/NNERuntimeRDG/Source/NNERuntimeRDG/Private/NNERuntimeRDGHlsl.cpp b/Engine/Plugins/Experimental/NNERuntimeRDG/Source/NNERuntimeRDG/Private/NNERuntimeRDGHlsl.cpp
--- a/Engine/Plugins/Experimental/NNERuntimeRDG/Source/NNERuntimeRDG/Private/NNERuntimeRDGHlsl.cpp
+++ b/Engine/Plugins/Experimental/NNERuntimeRDG/Source/NNERuntimeRDG/Private/NNERuntimeRDGHlsl.cpp
@@ -56,7 +56,7 @@
 using namespace UE::NNERuntimeRDG::Private::Hlsl;
 
 FGuid UNNERuntimeRDGHlslImpl::GUID = FGuid((int32)'R', (int32)'D', (int32)'G', (int32)'H');
-int32 UNNERuntimeRDGHlslImpl::Version = 0x00000006;
+int32 UNNERuntimeRDGHlslImpl::Version = 0x00000007;
 
 bool UNNERuntimeRDGHlslImpl::Init()
 {
@@ -146,7 +146,7 @@ bool UNNERuntimeRDGHlslImpl::IsCurrentPlatformSupported()
 	return bResult;
 }
 
-UNNERuntimeRDGHlslImpl::ECanCreateModelDataStatus UNNERuntimeRDGHlslImpl::CanCreateModelData(const FString& FileType, TConstArrayView<uint8> FileData, const TMap<FString, TConstArrayView<uint8>>& AdditionalFileData, const FGuid& FileId, const ITargetPlatform* TargetPlatform) const
+UNNERuntimeRDGHlslImpl::ECanCreateModelDataStatus UNNERuntimeRDGHlslImpl::CanCreateModelData(const FString& FileType, TConstArrayView64<uint8> FileData, const TMap<FString, TConstArrayView64<uint8>>& AdditionalFileData, const FGuid& FileId, const ITargetPlatform* TargetPlatform) const
 {
 #ifdef NNE_UTILITIES_AVAILABLE
 	return FileType.Compare("onnx", ESearchCase::IgnoreCase) == 0 ? ECanCreateModelDataStatus::Ok : ECanCreateModelDataStatus::FailFileIdNotSupported;
@@ -167,7 +167,7 @@ UNNERuntimeRDGHlslImpl::ECanCreateModelRDGStatus UNNERuntimeRDGHlslImpl::CanCrea
 		return ECanCreateModelRDGStatus::Fail;
 	}
 
-	TConstArrayView<uint8> Data = SharedData->GetView();
+	TConstArrayView64<uint8> Data = SharedData->GetView();
 
 	if (Data.Num() <= GuidSize + VersionSize)
 	{
@@ -181,11 +181,18 @@ UNNERuntimeRDGHlslImpl::ECanCreateModelRDGStatus UNNERuntimeRDGHlslImpl::CanCrea
 	return bResult ? ECanCreateModelRDGStatus::Ok : ECanCreateModelRDGStatus::Fail;
 };
 
-TSharedPtr<UE::NNE::FSharedModelData> UNNERuntimeRDGHlslImpl::CreateModelData(const FString& FileType, TConstArrayView<uint8> FileData, const TMap<FString, TConstArrayView<uint8>>& AdditionalFileData, const FGuid& FileId, const ITargetPlatform* TargetPlatform)
+TSharedPtr<UE::NNE::FSharedModelData> UNNERuntimeRDGHlslImpl::CreateModelData(const FString& FileType, TConstArrayView64<uint8> FileData, const TMap<FString, TConstArrayView64<uint8>>& AdditionalFileData, const FGuid& FileId, const ITargetPlatform* TargetPlatform)
 {
+	// Check model is not > 2GB
+	if ((TArray<uint8>::SizeType)FileData.Num() != FileData.Num())
+	{
+		UE_LOG(LogNNE, Warning, TEXT("UNNERuntimeRDGHlsl cannot create the model data with id %s (Filetype: %s), models > 2GBs are not supported"), *FileId.ToString(EGuidFormats::Digits).ToLower(), *FileType);
+		return {};
+	}
+
 	if (!AdditionalFileData.IsEmpty())
 	{
-		UE_LOG(LogNNE, Warning, TEXT("UNNERuntimeRDGHlsl cannot create the model data with id %s (Filetype: %s) external data not supported at the moment, please convert the model to internal storage. See https://onnx.ai/onnx/repo-docs/ExternalData.html"), *FileId.ToString(EGuidFormats::Digits).ToLower(), *FileType);
+		UE_LOG(LogNNE, Warning, TEXT("UNNERuntimeRDGHlsl cannot create the model data with id %s (Filetype: %s), external data not supported at the moment, please convert the model to internal storage. See https://onnx.ai/onnx/repo-docs/ExternalData.html"), *FileId.ToString(EGuidFormats::Digits).ToLower(), *FileType);
 		return {};
 	}
 
@@ -209,8 +216,8 @@ TSharedPtr<UE::NNE::FSharedModelData> UNNERuntimeRDGHlslImpl::CreateModelData(co
 		return {};
 	}
 
-	TArray<uint8> Result;
-	FMemoryWriter Writer(Result);
+	TArray64<uint8> Result;
+	FMemoryWriter64 Writer(Result);
 	
 	Writer << GUID;
 	Writer << Version;
@@ -222,7 +229,7 @@ TSharedPtr<UE::NNE::FSharedModelData> UNNERuntimeRDGHlslImpl::CreateModelData(co
 #endif //NNE_UTILITIES_AVAILABLE
 };
 
-FString UNNERuntimeRDGHlslImpl::GetModelDataIdentifier(const FString& FileType, TConstArrayView<uint8> FileData, const TMap<FString, TConstArrayView<uint8>>& AdditionalFileData, const FGuid& FileId, const ITargetPlatform* TargetPlatform) const
+FString UNNERuntimeRDGHlslImpl::GetModelDataIdentifier(const FString& FileType, TConstArrayView64<uint8> FileData, const TMap<FString, TConstArrayView64<uint8>>& AdditionalFileData, const FGuid& FileId, const ITargetPlatform* TargetPlatform) const
 {
 	return FileId.ToString(EGuidFormats::Digits) + "-" + UNNERuntimeRDGHlslImpl::GUID.ToString(EGuidFormats::Digits) + "-" + FString::FromInt(UNNERuntimeRDGHlslImpl::Version);
 }
diff --git a/Engine/Plugins/Experimental/NNERuntimeRDG/Source/NNERuntimeRDG/Private/NNERuntimeRDGHlsl.h b/Engine/Plugins/Experimental/NNERuntimeRDG/Source/NNERuntimeRDG/Private/NNERuntimeRDGHlsl.h
--- a/Engine/Plugins/Experimental/NNERuntimeRDG/Source/NNERuntimeRDG/Private/NNERuntimeRDGHlsl.h
+++ b/Engine/Plugins/Experimental/NNERuntimeRDG/Source/NNERuntimeRDG/Private/NNERuntimeRDGHlsl.h
@@ -30,9 +30,9 @@ public:
 
 	static bool IsCurrentPlatformSupported();
 
-	virtual ECanCreateModelDataStatus CanCreateModelData(const FString& FileType, TConstArrayView<uint8> FileData, const TMap<FString, TConstArrayView<uint8>>& AdditionalFileData, const FGuid& FileId, const ITargetPlatform* TargetPlatform) const override;
-	virtual TSharedPtr<UE::NNE::FSharedModelData> CreateModelData(const FString& FileType, TConstArrayView<uint8> FileData, const TMap<FString, TConstArrayView<uint8>>& AdditionalFileData, const FGuid& FileId, const ITargetPlatform* TargetPlatform) override;
-	virtual FString GetModelDataIdentifier(const FString& FileType, TConstArrayView<uint8> FileData, const TMap<FString, TConstArrayView<uint8>>& AdditionalFileData, const FGuid& FileId, const ITargetPlatform* TargetPlatform) const override;
+	virtual ECanCreateModelDataStatus CanCreateModelData(const FString& FileType, TConstArrayView64<uint8> FileData, const TMap<FString, TConstArrayView64<uint8>>& AdditionalFileData, const FGuid& FileId, const ITargetPlatform* TargetPlatform) const override;
+	virtual TSharedPtr<UE::NNE::FSharedModelData> CreateModelData(const FString& FileType, TConstArrayView64<uint8> FileData, const TMap<FString, TConstArrayView64<uint8>>& AdditionalFileData, const FGuid& FileId, const ITargetPlatform* TargetPlatform) override;
+	virtual FString GetModelDataIdentifier(const FString& FileType, TConstArrayView64<uint8> FileData, const TMap<FString, TConstArrayView64<uint8>>& AdditionalFileData, const FGuid& FileId, const ITargetPlatform* TargetPlatform) const override;
 
 	virtual ECanCreateModelRDGStatus CanCreateModelRDG(const TObjectPtr<UNNEModelData> ModelData) const override;
 	virtual TSharedPtr<UE::NNE::IModelRDG> CreateModelRDG(const TObjectPtr<UNNEModelData> ModelData) override;
diff --git a/Engine/Plugins/NNE/NNERuntimeORT/Source/NNERuntimeORT/Private/NNERuntimeORT.cpp b/Engine/Plugins/NNE/NNERuntimeORT/Source/NNERuntimeORT/Private/NNERuntimeORT.cpp
--- a/Engine/Plugins/NNE/NNERuntimeORT/Source/NNERuntimeORT/Private/NNERuntimeORT.cpp
+++ b/Engine/Plugins/NNE/NNERuntimeORT/Source/NNERuntimeORT/Private/NNERuntimeORT.cpp
@@ -20,10 +20,10 @@
 #include UE_INLINE_GENERATED_CPP_BY_NAME(NNERuntimeORT)
 
 FGuid UNNERuntimeORTCpu::GUID = FGuid((int32)'O', (int32)'C', (int32)'P', (int32)'U');
-int32 UNNERuntimeORTCpu::Version = 0x00000003;
+int32 UNNERuntimeORTCpu::Version = 0x00000004;
 
 FGuid UNNERuntimeORTDml::GUID = FGuid((int32)'O', (int32)'D', (int32)'M', (int32)'L');
-int32 UNNERuntimeORTDml::Version = 0x00000003;
+int32 UNNERuntimeORTDml::Version = 0x00000004;
 
 namespace UE::NNERuntimeORT::Private::Details
 { 
@@ -31,7 +31,7 @@ namespace UE::NNERuntimeORT::Private::Details
 	static FString OnnxExternalDataDescriptorKey(TEXT("OnnxExternalDataDescriptor"));
 	static FString OnnxExternalDataBytesKey(TEXT("OnnxExternalDataBytes"));
 
-	FOnnxDataDescriptor MakeOnnxDataDescriptor(TConstArrayView<uint8> FileData, const TMap<FString, TConstArrayView<uint8>>& AdditionalFileData)
+	FOnnxDataDescriptor MakeOnnxDataDescriptor(TConstArrayView64<uint8> FileData, const TMap<FString, TConstArrayView64<uint8>>& AdditionalFileData)
 	{
 		FOnnxDataDescriptor OnnxDataDescriptor = {};
 		OnnxDataDescriptor.OnnxModelDataSize = FileData.Num();
@@ -60,7 +60,7 @@ namespace UE::NNERuntimeORT::Private::Details
 		return OnnxDataDescriptor;
 	}
 
-	void WriteOnnxModelData(FMemoryWriter Writer, TConstArrayView<uint8> FileData, const TMap<FString, TConstArrayView<uint8>>& AdditionalFileData)
+	void WriteOnnxModelData(FMemoryWriter64 Writer, TConstArrayView64<uint8> FileData, const TMap<FString, TConstArrayView64<uint8>>& AdditionalFileData)
 	{
 		FOnnxDataDescriptor Descriptor = MakeOnnxDataDescriptor(FileData, AdditionalFileData);
 		check(FileData.Num() == Descriptor.OnnxModelDataSize);
@@ -78,12 +78,12 @@ namespace UE::NNERuntimeORT::Private::Details
 	}
 } // namespace UE::NNERuntimeORT::Private::Details
 
-UNNERuntimeORTCpu::ECanCreateModelDataStatus UNNERuntimeORTCpu::CanCreateModelData(const FString& FileType, TConstArrayView<uint8> FileData, const TMap<FString, TConstArrayView<uint8>>& AdditionalFileData, const FGuid& FileId, const ITargetPlatform* TargetPlatform) const
+UNNERuntimeORTCpu::ECanCreateModelDataStatus UNNERuntimeORTCpu::CanCreateModelData(const FString& FileType, TConstArrayView64<uint8> FileData, const TMap<FString, TConstArrayView64<uint8>>& AdditionalFileData, const FGuid& FileId, const ITargetPlatform* TargetPlatform) const
 {
 	return (!FileData.IsEmpty() && FileType.Compare("onnx", ESearchCase::IgnoreCase) == 0) ? ECanCreateModelDataStatus::Ok : ECanCreateModelDataStatus::FailFileIdNotSupported;
 }
 
-TSharedPtr<UE::NNE::FSharedModelData> UNNERuntimeORTCpu::CreateModelData(const FString& FileType, TConstArrayView<uint8> FileData, const TMap<FString, TConstArrayView<uint8>>& AdditionalFileData, const FGuid& FileId, const ITargetPlatform* TargetPlatform)
+TSharedPtr<UE::NNE::FSharedModelData> UNNERuntimeORTCpu::CreateModelData(const FString& FileType, TConstArrayView64<uint8> FileData, const TMap<FString, TConstArrayView64<uint8>>& AdditionalFileData, const FGuid& FileId, const ITargetPlatform* TargetPlatform)
 {
 	using namespace UE::NNERuntimeORT::Private;
 
@@ -93,35 +93,39 @@ TSharedPtr<UE::NNE::FSharedModelData> UNNERuntimeORTCpu::CreateModelData(const F
 		return {};
 	}
 
-	FNNEModelRaw InputModel{ TArray<uint8>{FileData}, ENNEInferenceFormat::ONNX };
+	TConstArrayView64<uint8> OptimizedModelView = FileData;
+	TArray64<uint8> OptimizedModelBuffer;
 
 	//For now only optimize model if there is no external data (as additional data are serialized from the unoptimized model below)
 	if (AdditionalFileData.IsEmpty())
 	{
+		
 		if (GraphOptimizationLevel OptimizationLevel = GetGraphOptimizationLevelForCPU(false, IsRunningCookCommandlet()); OptimizationLevel > GraphOptimizationLevel::ORT_DISABLE_ALL)
 		{
 			TUniquePtr<Ort::SessionOptions> SessionOptions = CreateSessionOptionsDefault(Environment.ToSharedRef());
 			SessionOptions->SetGraphOptimizationLevel(OptimizationLevel);
 			SessionOptions->EnableCpuMemArena();
 
-			if (!OptimizeModel(Environment.ToSharedRef(), *SessionOptions, ENNEInferenceFormat::ONNX, InputModel))
+			if (!OptimizeModel(Environment.ToSharedRef(), *SessionOptions, FileData, OptimizedModelBuffer))
 			{
 				return {};
 			}
+
+			OptimizedModelView = OptimizedModelBuffer;
 		}
 	}
 
-	TArray<uint8> Result;
-	FMemoryWriter Writer(Result, /*bIsPersitent =*/ true);
+	TArray64<uint8> Result;
+	FMemoryWriter64 Writer(Result, /*bIsPersitent =*/ true);
 	Writer << UNNERuntimeORTCpu::GUID;
 	Writer << UNNERuntimeORTCpu::Version;
 
-	Details::WriteOnnxModelData(Writer, InputModel.Data, AdditionalFileData);
+	Details::WriteOnnxModelData(Writer, OptimizedModelView, AdditionalFileData);
 
 	return MakeShared<UE::NNE::FSharedModelData>(MakeSharedBufferFromArray(MoveTemp(Result)), 0);
 }
 
-FString UNNERuntimeORTCpu::GetModelDataIdentifier(const FString& FileType, TConstArrayView<uint8> FileData, const TMap<FString, TConstArrayView<uint8>>& AdditionalFileData, const FGuid& FileId, const ITargetPlatform* TargetPlatform) const
+FString UNNERuntimeORTCpu::GetModelDataIdentifier(const FString& FileType, TConstArrayView64<uint8> FileData, const TMap<FString, TConstArrayView64<uint8>>& AdditionalFileData, const FGuid& FileId, const ITargetPlatform* TargetPlatform) const
 {
 	return FileId.ToString(EGuidFormats::Digits) + "-" + UNNERuntimeORTCpu::GUID.ToString(EGuidFormats::Digits) + "-" + FString::FromInt(UNNERuntimeORTCpu::Version);
 }
@@ -149,7 +153,7 @@ UNNERuntimeORTCpu::ECanCreateModelCPUStatus UNNERuntimeORTCpu::CanCreateModelCPU
 		return ECanCreateModelCPUStatus::Fail;
 	}
 
-	TConstArrayView<uint8> Data = SharedData->GetView();
+	TConstArrayView64<uint8> Data = SharedData->GetView();
 
 	if (Data.Num() <= GuidSize + VersionSize)
 	{
@@ -206,12 +210,12 @@ FString UNNERuntimeORTDml::GetRuntimeName() const
 	return TEXT("NNERuntimeORTDml");
 }
 
-UNNERuntimeORTDml::ECanCreateModelDataStatus UNNERuntimeORTDml::CanCreateModelData(const FString& FileType, TConstArrayView<uint8> FileData, const TMap<FString, TConstArrayView<uint8>>& AdditionalFileData, const FGuid& FileId, const ITargetPlatform* TargetPlatform) const
+UNNERuntimeORTDml::ECanCreateModelDataStatus UNNERuntimeORTDml::CanCreateModelData(const FString& FileType, TConstArrayView64<uint8> FileData, const TMap<FString, TConstArrayView64<uint8>>& AdditionalFileData, const FGuid& FileId, const ITargetPlatform* TargetPlatform) const
 {
 	return (!FileData.IsEmpty() && FileType.Compare("onnx", ESearchCase::IgnoreCase) == 0) ? ECanCreateModelDataStatus::Ok : ECanCreateModelDataStatus::FailFileIdNotSupported;
 }
 
-TSharedPtr<UE::NNE::FSharedModelData> UNNERuntimeORTDml::CreateModelData(const FString& FileType, TConstArrayView<uint8> FileData, const TMap<FString, TConstArrayView<uint8>>& AdditionalFileData, const FGuid& FileId, const ITargetPlatform* TargetPlatform)
+TSharedPtr<UE::NNE::FSharedModelData> UNNERuntimeORTDml::CreateModelData(const FString& FileType, TConstArrayView64<uint8> FileData, const TMap<FString, TConstArrayView64<uint8>>& AdditionalFileData, const FGuid& FileId, const ITargetPlatform* TargetPlatform)
 {
 	using namespace UE::NNERuntimeORT::Private;
 
@@ -221,7 +225,8 @@ TSharedPtr<UE::NNE::FSharedModelData> UNNERuntimeORTDml::CreateModelData(const F
 		return {};
 	}
 
-	FNNEModelRaw InputModel{ TArray<uint8>{FileData}, ENNEInferenceFormat::ONNX };
+	TConstArrayView64<uint8> OptimizedModelView = FileData;
+	TArray64<uint8> OptimizedModelBuffer;
 
 	//For now only optimize model if there is no external data (as additional data are serialized from the unoptimized model below)
 	if (AdditionalFileData.IsEmpty())
@@ -233,24 +238,26 @@ TSharedPtr<UE::NNE::FSharedModelData> UNNERuntimeORTDml::CreateModelData(const F
 			SessionOptions->SetExecutionMode(ExecutionMode::ORT_SEQUENTIAL);
 			SessionOptions->DisableMemPattern();
 
-			if (!OptimizeModel(Environment.ToSharedRef(), *SessionOptions, ENNEInferenceFormat::ONNX, InputModel))
+			if (!OptimizeModel(Environment.ToSharedRef(), *SessionOptions, FileData, OptimizedModelBuffer))
 			{
 				return {};
 			}
+
+			OptimizedModelView = OptimizedModelBuffer;
 		}
 	}
 
-	TArray<uint8> Result;
-	FMemoryWriter Writer(Result, /*bIsPersitent =*/ true);
+	TArray64<uint8> Result;
+	FMemoryWriter64 Writer(Result, /*bIsPersitent =*/ true);
 	Writer << UNNERuntimeORTDml::GUID;
 	Writer << UNNERuntimeORTDml::Version;
 
-	Details::WriteOnnxModelData(Writer, InputModel.Data, AdditionalFileData);
+	Details::WriteOnnxModelData(Writer, OptimizedModelView, AdditionalFileData);
 
 	return MakeShared<UE::NNE::FSharedModelData>(MakeSharedBufferFromArray(MoveTemp(Result)), 0);
 }
 
-FString UNNERuntimeORTDml::GetModelDataIdentifier(const FString& FileType, TConstArrayView<uint8> FileData, const TMap<FString, TConstArrayView<uint8>>& AdditionalFileData, const FGuid& FileId, const ITargetPlatform* TargetPlatform) const
+FString UNNERuntimeORTDml::GetModelDataIdentifier(const FString& FileType, TConstArrayView64<uint8> FileData, const TMap<FString, TConstArrayView64<uint8>>& AdditionalFileData, const FGuid& FileId, const ITargetPlatform* TargetPlatform) const
 {
 	return FileId.ToString(EGuidFormats::Digits) + "-" + UNNERuntimeORTDml::GUID.ToString(EGuidFormats::Digits) + "-" + FString::FromInt(UNNERuntimeORTDml::Version);
 }
@@ -354,7 +361,7 @@ UNNERuntimeORTDml::ECanCreateModelCommonStatus UNNERuntimeORTDml::CanCreateModel
 		return ECanCreateModelCommonStatus::Fail;
 	}
 
-	TConstArrayView<uint8> Data = SharedData->GetView();
+	TConstArrayView64<uint8> Data = SharedData->GetView();
 
 	if (Data.Num() <= GuidSize + VersionSize)
 	{
diff --git a/Engine/Plugins/NNE/NNERuntimeORT/Source/NNERuntimeORT/Private/NNERuntimeORT.h b/Engine/Plugins/NNE/NNERuntimeORT/Source/NNERuntimeORT/Private/NNERuntimeORT.h
--- a/Engine/Plugins/NNE/NNERuntimeORT/Source/NNERuntimeORT/Private/NNERuntimeORT.h
+++ b/Engine/Plugins/NNE/NNERuntimeORT/Source/NNERuntimeORT/Private/NNERuntimeORT.h
@@ -35,9 +35,9 @@ public:
 
 	virtual FString GetRuntimeName() const override;
 	
-	virtual ECanCreateModelDataStatus CanCreateModelData(const FString& FileType, TConstArrayView<uint8> FileData, const TMap<FString, TConstArrayView<uint8>>& AdditionalFileData, const FGuid& FileId, const ITargetPlatform* TargetPlatform) const override;
-	virtual TSharedPtr<UE::NNE::FSharedModelData> CreateModelData(const FString& FileType, TConstArrayView<uint8> FileData, const TMap<FString, TConstArrayView<uint8>>& AdditionalFileData, const FGuid& FileId, const ITargetPlatform* TargetPlatform) override;
-	virtual FString GetModelDataIdentifier(const FString& FileType, TConstArrayView<uint8> FileData, const TMap<FString, TConstArrayView<uint8>>& AdditionalFileData, const FGuid& FileId, const ITargetPlatform* TargetPlatform) const override;
+	virtual ECanCreateModelDataStatus CanCreateModelData(const FString& FileType, TConstArrayView64<uint8> FileData, const TMap<FString, TConstArrayView64<uint8>>& AdditionalFileData, const FGuid& FileId, const ITargetPlatform* TargetPlatform) const override;
+	virtual TSharedPtr<UE::NNE::FSharedModelData> CreateModelData(const FString& FileType, TConstArrayView64<uint8> FileData, const TMap<FString, TConstArrayView64<uint8>>& AdditionalFileData, const FGuid& FileId, const ITargetPlatform* TargetPlatform) override;
+	virtual FString GetModelDataIdentifier(const FString& FileType, TConstArrayView64<uint8> FileData, const TMap<FString, TConstArrayView64<uint8>>& AdditionalFileData, const FGuid& FileId, const ITargetPlatform* TargetPlatform) const override;
 
 	virtual ECanCreateModelCPUStatus CanCreateModelCPU(const TObjectPtr<UNNEModelData> ModelData) const override;
 	virtual TSharedPtr<UE::NNE::IModelCPU> CreateModelCPU(const TObjectPtr<UNNEModelData> ModelData) override;
@@ -63,9 +63,9 @@ public:
 
 	virtual FString GetRuntimeName() const override;
 
-	virtual ECanCreateModelDataStatus CanCreateModelData(const FString& FileType, TConstArrayView<uint8> FileData, const TMap<FString, TConstArrayView<uint8>>& AdditionalFileData, const FGuid& FileId, const ITargetPlatform* TargetPlatform) const override;
-	virtual TSharedPtr<UE::NNE::FSharedModelData> CreateModelData(const FString& FileType, TConstArrayView<uint8> FileData, const TMap<FString, TConstArrayView<uint8>>& AdditionalFileData, const FGuid& FileId, const ITargetPlatform* TargetPlatform) override;
-	virtual FString GetModelDataIdentifier(const FString& FileType, TConstArrayView<uint8> FileData, const TMap<FString, TConstArrayView<uint8>>& AdditionalFileData, const FGuid& FileId, const ITargetPlatform* TargetPlatform) const override;
+	virtual ECanCreateModelDataStatus CanCreateModelData(const FString& FileType, TConstArrayView64<uint8> FileData, const TMap<FString, TConstArrayView64<uint8>>& AdditionalFileData, const FGuid& FileId, const ITargetPlatform* TargetPlatform) const override;
+	virtual TSharedPtr<UE::NNE::FSharedModelData> CreateModelData(const FString& FileType, TConstArrayView64<uint8> FileData, const TMap<FString, TConstArrayView64<uint8>>& AdditionalFileData, const FGuid& FileId, const ITargetPlatform* TargetPlatform) override;
+	virtual FString GetModelDataIdentifier(const FString& FileType, TConstArrayView64<uint8> FileData, const TMap<FString, TConstArrayView64<uint8>>& AdditionalFileData, const FGuid& FileId, const ITargetPlatform* TargetPlatform) const override;
 
 	virtual ECanCreateModelGPUStatus CanCreateModelGPU(const TObjectPtr<UNNEModelData> ModelData) const override;
 	virtual TSharedPtr<UE::NNE::IModelGPU> CreateModelGPU(const TObjectPtr<UNNEModelData> ModelData) override;
diff --git a/Engine/Plugins/NNE/NNERuntimeORT/Source/NNERuntimeORT/Private/NNERuntimeORTModel.cpp b/Engine/Plugins/NNE/NNERuntimeORT/Source/NNERuntimeORT/Private/NNERuntimeORTModel.cpp
--- a/Engine/Plugins/NNE/NNERuntimeORT/Source/NNERuntimeORT/Private/NNERuntimeORTModel.cpp
+++ b/Engine/Plugins/NNE/NNERuntimeORT/Source/NNERuntimeORT/Private/NNERuntimeORTModel.cpp
@@ -61,12 +61,12 @@ namespace Detail
 	}
 
 	bool CreateSession(
-		TConstArrayView<uint8> ModelData,
+		TConstArrayView64<uint8> ModelData,
 		const Ort::SessionOptions& SessionOptions,
 		const FEnvironment& Environment,
 		TUniquePtr<Ort::Session>& Session, FString& TempDirForModelWithExternalData)
 	{
-		FMemoryReaderView Reader(ModelData, /*bIsPersitent =*/ true);
+		FMemoryReaderView Reader(MakeMemoryView(ModelData), /*bIsPersitent =*/ true);
 		FGuid GUID;
 		int32 Version;
 		Reader << GUID;
@@ -76,7 +76,7 @@ namespace Detail
 		Reader << Descriptor;
 
 		int64 BaseDataOffset = Reader.Tell();
-		TConstArrayView<uint8> ModelBuffer = TConstArrayView<uint8>(&(ModelData.GetData()[BaseDataOffset]), Descriptor.OnnxModelDataSize);
+		TConstArrayView64<uint8> ModelBuffer = TConstArrayView64<uint8>(&(ModelData.GetData()[BaseDataOffset]), Descriptor.OnnxModelDataSize);
 
 		if (ModelBuffer.Num() == 0)
 		{
@@ -99,19 +99,19 @@ namespace Detail
 				if (!FFileHelper::SaveArrayToFile(ModelBuffer, *Filepath))
 				{
 					IFileManager::Get().DeleteDirectory(*TempDirForModelWithExternalData, false, true);
-					UE_LOG(LogNNE, Error, TEXT("NNERuntimeORT::Private::Details::CreateSession() could not write model to disk at %s."), *Filepath);
+					UE_LOG(LogNNE, Error, TEXT("Large models are an experimental feature at the moment. NNERuntimeORT::Private::Details::CreateSession() could not write model to disk at %s."), *Filepath);
 					return false;
 				}
 
 				for (const FOnnxAdditionalDataDescriptor& AdditionalDataDescriptor : Descriptor.AdditionalDataDescriptors)
 				{
 					FString AdditionalDataFilename = FPaths::Combine(TempDirForModelWithExternalData, *AdditionalDataDescriptor.Path);
-					TConstArrayView<uint8> AdditionalDataBuffer = TConstArrayView<uint8>(&(ModelData.GetData()[BaseDataOffset + AdditionalDataDescriptor.Offset]), AdditionalDataDescriptor.Size);
+					TConstArrayView64<uint8> AdditionalDataBuffer = TConstArrayView64<uint8>(&(ModelData.GetData()[BaseDataOffset + AdditionalDataDescriptor.Offset]), AdditionalDataDescriptor.Size);
 
 					if (!FFileHelper::SaveArrayToFile(AdditionalDataBuffer, *AdditionalDataFilename))
 					{
 						IFileManager::Get().DeleteDirectory(*TempDirForModelWithExternalData, false, true);
-						UE_LOG(LogNNE, Error, TEXT("NNERuntimeORT::Private::Details::CreateSession() could not write additional data to disk at %s."), *AdditionalDataFilename);
+						UE_LOG(LogNNE, Error, TEXT("Large models are an experimental feature at the moment. NNERuntimeORT::Private::Details::CreateSession() could not write additional data to disk at %s."), *AdditionalDataFilename);
 						return false;
 					}
 				}
@@ -152,13 +152,13 @@ FModelInstanceORTBase<ModelInterface, TensorBinding>::~FModelInstanceORTBase()
 	{
 		if (!IFileManager::Get().DeleteDirectory(*TempDirForModelWithExternalData, false, true))
 		{
-			UE_LOG(LogNNE, Warning, TEXT("FModelInstanceORTBase could not delete temp directy %s on model instance destruction."), *TempDirForModelWithExternalData);
+			UE_LOG(LogNNE, Warning, TEXT("Large models are an experimental feature at the moment. FModelInstanceORTBase could not delete temp directy %s on model instance destruction."), *TempDirForModelWithExternalData);
 		}
 	}
 }
 
 template <class ModelInterface, class TensorBinding> 
-bool FModelInstanceORTBase<ModelInterface, TensorBinding>::Init(TConstArrayView<uint8> ModelData)
+bool FModelInstanceORTBase<ModelInterface, TensorBinding>::Init(TConstArrayView64<uint8> ModelData)
 {
 #if WITH_EDITOR
 	try
@@ -544,7 +544,7 @@ FModelInstanceORTDmlRDG::~FModelInstanceORTDmlRDG()
 	{
 		if (!IFileManager::Get().DeleteDirectory(*TempDirForModelWithExternalData, false, true))
 		{
-			UE_LOG(LogNNE, Warning, TEXT("FModelInstanceORTDmlRDG could not delete temp directy %s on model instance destruction."), *TempDirForModelWithExternalData);
+			UE_LOG(LogNNE, Warning, TEXT("Large models are an experimental feature at the moment. FModelInstanceORTDmlRDG could not delete temp directy %s on model instance destruction."), *TempDirForModelWithExternalData);
 		}
 	}
 }
diff --git a/Engine/Plugins/NNE/NNERuntimeORT/Source/NNERuntimeORT/Private/NNERuntimeORTModel.h b/Engine/Plugins/NNE/NNERuntimeORT/Source/NNERuntimeORT/Private/NNERuntimeORTModel.h
--- a/Engine/Plugins/NNE/NNERuntimeORT/Source/NNERuntimeORT/Private/NNERuntimeORTModel.h
+++ b/Engine/Plugins/NNE/NNERuntimeORT/Source/NNERuntimeORT/Private/NNERuntimeORTModel.h
@@ -31,7 +31,7 @@ public:
 
 	virtual typename ModelInterface::ESetInputTensorShapesStatus SetInputTensorShapes(TConstArrayView<NNE::FTensorShape> InInputShapes) override;
 
-	bool Init(TConstArrayView<uint8> ModelData);
+	bool Init(TConstArrayView64<uint8> ModelData);
 
 	typename ModelInterface::ERunSyncStatus RunSync(TConstArrayView<TensorBinding> InInputBindings, TConstArrayView<TensorBinding> InOutputBindings) override;
 
diff --git a/Engine/Plugins/NNE/NNERuntimeORT/Source/NNERuntimeORT/Private/NNERuntimeORTUtils.cpp b/Engine/Plugins/NNE/NNERuntimeORT/Source/NNERuntimeORT/Private/NNERuntimeORTUtils.cpp
--- a/Engine/Plugins/NNE/NNERuntimeORT/Source/NNERuntimeORT/Private/NNERuntimeORTUtils.cpp
+++ b/Engine/Plugins/NNE/NNERuntimeORT/Source/NNERuntimeORT/Private/NNERuntimeORTUtils.cpp
@@ -307,19 +307,13 @@ TUniquePtr<Ort::SessionOptions> CreateSessionOptionsForDirectML(const TSharedRef
 #endif // PLATFORM_WINDOWS
 }
 
-bool OptimizeModel(const TSharedRef<FEnvironment> &Environment, Ort::SessionOptions &SessionOptions, ENNEInferenceFormat TargetFormat, FNNEModelRaw& Model)
+bool OptimizeModel(const TSharedRef<FEnvironment> &Environment, Ort::SessionOptions &SessionOptions, 
+					TConstArrayView64<uint8>& InputModel, TArray64<uint8>& OptimizedModel)
 {
 	SCOPED_NAMED_EVENT_TEXT("OrtHelper::OptimizeModel", FColor::Magenta);
 
-	if (Model.Format != ENNEInferenceFormat::ONNX)
-	{
-		UE_LOG(LogNNE, Warning, TEXT("NNERuntimeORT: ONNX Runtime Model Optimizer is expecting a model in ONNX format but received %u."), Model.Format);
-		return false;
-	}
-
 	FString ProjIntermediateDir = FPaths::ConvertRelativePathToFull(FPaths::ProjectIntermediateDir());
-	FString TargetExtension = TargetFormat == ENNEInferenceFormat::ONNX ? TEXT(".onnx") : TEXT(".ort");
-	FString ModelOptimizedPath = FPaths::CreateTempFilename(*ProjIntermediateDir, TEXT("ORTOptimizerPass_Optimized"), *TargetExtension);
+	FString ModelOptimizedPath = FPaths::CreateTempFilename(*ProjIntermediateDir, TEXT("ORTOptimizerPass_Optimized"), TEXT(".onnx"));
 
 #if WITH_EDITOR
 	try
@@ -328,11 +322,11 @@ bool OptimizeModel(const TSharedRef<FEnvironment> &Environment, Ort::SessionOpti
 #if PLATFORM_WINDOWS
 		SessionOptions.SetOptimizedModelFilePath(*ModelOptimizedPath);
 
-		Ort::Session Session(Environment->GetOrtEnv(), Model.Data.GetData(), Model.Data.Num(), SessionOptions);
+		Ort::Session Session(Environment->GetOrtEnv(), InputModel.GetData(), InputModel.Num(), SessionOptions);
 #else
 		SessionOptions.SetOptimizedModelFilePath(TCHAR_TO_ANSI(*ModelOptimizedPath));
 
-		Ort::Session Session(Environment->GetOrtEnv(), Model.Data.GetData(), Model.Data.Num(), SessionOptions);
+		Ort::Session Session(Environment->GetOrtEnv(), InputModel.GetData(), InputModel.Num(), SessionOptions);
 #endif
 	}
 #if WITH_EDITOR
@@ -348,12 +342,10 @@ bool OptimizeModel(const TSharedRef<FEnvironment> &Environment, Ort::SessionOpti
 	}
 #endif // WITH_EDITOR
 
-	FFileHelper::LoadFileToArray(Model.Data, *ModelOptimizedPath);
+	FFileHelper::LoadFileToArray(OptimizedModel, *ModelOptimizedPath);
 
 	IFileManager::Get().Delete(*ModelOptimizedPath);
 
-	Model.Format = TargetFormat;
-
 	return true;
 }
 
diff --git a/Engine/Plugins/NNE/NNERuntimeORT/Source/NNERuntimeORT/Private/NNERuntimeORTUtils.h b/Engine/Plugins/NNE/NNERuntimeORT/Source/NNERuntimeORT/Private/NNERuntimeORTUtils.h
--- a/Engine/Plugins/NNE/NNERuntimeORT/Source/NNERuntimeORT/Private/NNERuntimeORTUtils.h
+++ b/Engine/Plugins/NNE/NNERuntimeORT/Source/NNERuntimeORT/Private/NNERuntimeORTUtils.h
@@ -28,7 +28,8 @@ namespace OrtHelper
 
 	TUniquePtr<Ort::SessionOptions> CreateSessionOptionsForDirectML(const TSharedRef<FEnvironment> &Environment, bool bRHID3D12Required = true);
 
-	bool OptimizeModel(const TSharedRef<FEnvironment> &Environment, Ort::SessionOptions &SessionOptions, ENNEInferenceFormat TargetFormat, FNNEModelRaw& Model);
+	bool OptimizeModel(const TSharedRef<FEnvironment>& Environment, Ort::SessionOptions& SessionOptions, 
+		TConstArrayView64<uint8>& InputModel, TArray64<uint8>& OptimizedModel);
 
 	struct TypeInfoORT
 	{
diff --git a/Engine/Source/Editor/NNEEditor/Private/NNEEditorOnnxFileLoaderHelper.cpp b/Engine/Source/Editor/NNEEditor/Private/NNEEditorOnnxFileLoaderHelper.cpp
--- a/Engine/Source/Editor/NNEEditor/Private/NNEEditorOnnxFileLoaderHelper.cpp
+++ b/Engine/Source/Editor/NNEEditor/Private/NNEEditorOnnxFileLoaderHelper.cpp
@@ -71,7 +71,7 @@ namespace UE::NNEEditor::Internal::OnnxFileLoaderHelper
 
 		ModelFileSize = OnnxData.Num();
 
-		TMap<FString, TConstArrayView<uint8>> AdditionalBuffers;
+		TMap<FString, TConstArrayView64<uint8>> AdditionalBuffers;
 		TArray64<uint8> OnnxExternalDataBytesBuffer;
 		TArray64<uint8> OnnxExternalDataDescriptorBuffer;
 
diff --git a/Engine/Source/Runtime/Core/Private/Misc/FileHelper.cpp b/Engine/Source/Runtime/Core/Private/Misc/FileHelper.cpp
--- a/Engine/Source/Runtime/Core/Private/Misc/FileHelper.cpp
+++ b/Engine/Source/Runtime/Core/Private/Misc/FileHelper.cpp
@@ -571,7 +571,7 @@ bool FFileHelper::LoadFileToStringWithLineVisitor(const TCHAR* Filename, TFuncti
 /**
  * Save a binary array to a file.
  */
-bool FFileHelper::SaveArrayToFile(TArrayView<const uint8> Array, const TCHAR* Filename, IFileManager* FileManager /*= &IFileManager::Get()*/, uint32 WriteFlags)
+bool FFileHelper::SaveArrayToFile(TArrayView64<const uint8> Array, const TCHAR* Filename, IFileManager* FileManager /*= &IFileManager::Get()*/, uint32 WriteFlags)
 {
 	TUniquePtr<FArchive> Ar = TUniquePtr<FArchive>( FileManager->CreateFileWriter( Filename, WriteFlags ) );
 	if( !Ar )
diff --git a/Engine/Source/Runtime/Core/Public/Misc/FileHelper.h b/Engine/Source/Runtime/Core/Public/Misc/FileHelper.h
--- a/Engine/Source/Runtime/Core/Public/Misc/FileHelper.h
+++ b/Engine/Source/Runtime/Core/Public/Misc/FileHelper.h
@@ -139,7 +139,7 @@ struct FFileHelper
 	/**
 	 * Save a binary array to a file.
 	 */
-	static CORE_API bool SaveArrayToFile(TArrayView<const uint8> Array, const TCHAR* Filename, IFileManager* FileManager=&IFileManager::Get(), uint32 WriteFlags = 0);
+	static CORE_API bool SaveArrayToFile(TArrayView64<const uint8> Array, const TCHAR* Filename, IFileManager* FileManager = &IFileManager::Get(), uint32 WriteFlags = 0);
 
 	/**
 	 * Save a binary array to a file.
diff --git a/Engine/Source/Runtime/NNE/Internal/NNERuntimeFormat.h b/Engine/Source/Runtime/NNE/Internal/NNERuntimeFormat.h
--- a/Engine/Source/Runtime/NNE/Internal/NNERuntimeFormat.h
+++ b/Engine/Source/Runtime/NNE/Internal/NNERuntimeFormat.h
@@ -111,7 +111,7 @@ struct FNNERuntimeFormat
 	TArray<FNNEFormatOperatorDesc> Operators;
 
 	uint64 DataSize;
-	TArray<uint8> TensorData;
+	TArray64<uint8> TensorData;
 	
 	bool Serialize(FArchive& Ar)
 	{
diff --git a/Engine/Source/Runtime/NNE/Private/NNEModelData.cpp b/Engine/Source/Runtime/NNE/Private/NNEModelData.cpp
--- a/Engine/Source/Runtime/NNE/Private/NNEModelData.cpp
+++ b/Engine/Source/Runtime/NNE/Private/NNEModelData.cpp
@@ -18,6 +18,7 @@ namespace UE::NNE::ModelData
 		V1 = 1, // TargetRuntimes and AssetImportData
 		V2 = 2, // Re-arrange fields and store only ModelData in cooked assets
 		V3 = 3, // Adding AdditionalFileData
+		V4 = 4, // Support for > 2GB models
 		// New versions can be added above this line
 		VersionPlusOne,
 		Latest = VersionPlusOne - 1
@@ -48,7 +49,7 @@ namespace UE::NNE::ModelData
 		return RuntimesAsOneString;
 	}
 
-	inline TSharedPtr<UE::NNE::FSharedModelData> CreateModelData(const FString& RuntimeName, const FString& FileType, const TArray<uint8>& FileData, const TMap<FString, TConstArrayView<uint8>>& AdditionalFileData, const FGuid& FileId, const ITargetPlatform* TargetPlatform)
+	inline TSharedPtr<UE::NNE::FSharedModelData> CreateModelData(const FString& RuntimeName, const FString& FileType, const TArray64<uint8>& FileData, const TMap<FString, TConstArrayView64<uint8>>& AdditionalFileData, const FGuid& FileId, const ITargetPlatform* TargetPlatform)
 	{
 		TWeakInterfacePtr<INNERuntime> NNERuntime = UE::NNE::GetRuntime<INNERuntime>(RuntimeName);
 		if (NNERuntime.IsValid())
@@ -94,9 +95,9 @@ namespace UE::NNE
 		MemoryAlignment = 0;
 	}
 
-	TConstArrayView<uint8> FSharedModelData::GetView() const
+	TConstArrayView64<uint8> FSharedModelData::GetView() const
 	{
-		return MakeArrayView(static_cast<const uint8*>(Data.GetData()), Data.GetSize());
+		return TConstArrayView64<uint8>(static_cast<const uint8*>(Data.GetData()), Data.GetSize());
 	}
 
 	uint32 FSharedModelData::GetMemoryAlignment() const
@@ -125,7 +126,7 @@ void UNNEModelData::Serialize(FArchive& Ar)
 			Ar << TmpTargetRuntimes;
 			FString TmpFileType;
 			Ar << TmpFileType;
-			TArray<uint8> TmpFileData;
+			TArray64<uint8> TmpFileData;
 			Ar << TmpFileData;
 			int32 NumAdditionalFileDataItems = 0;
 			Ar << NumAdditionalFileDataItems;
@@ -148,7 +149,7 @@ void UNNEModelData::Serialize(FArchive& Ar)
 					CookRuntimeNames.Append(GetTargetRuntimes());
 				}
 
-				TMap<FString, TConstArrayView<uint8>> AdditionalFileDataView;
+				TMap<FString, TConstArrayView64<uint8>> AdditionalFileDataView;
 				TSet<FString> Keys;
 				AdditionalFileData.GetKeys(Keys);
 				for (auto& Key : Keys)
@@ -225,7 +226,8 @@ void UNNEModelData::Serialize(FArchive& Ar)
 		FString Name;
 		uint32 MemoryAlignment = 0;
 		uint64 DataSize = 0;
-		TArray<uint8> Data;
+		TArray64<uint8> Data;
+		TArray<uint8> Data_Size32Bits;
 		void* RawData = nullptr;
 		int32 Index = 0;
 
@@ -234,14 +236,15 @@ void UNNEModelData::Serialize(FArchive& Ar)
 		case UE::NNE::ModelData::Version::V0:
 			TargetRuntimes.Empty();
 			Ar << FileType;
-			Ar << FileData;
+			Ar << Data_Size32Bits;
+			FileData = MoveTemp(Data_Size32Bits);
 			Ar << FileId;
 			Ar << NumItems;
 			for (Index = 0; Index < NumItems; Index++)
 			{
 				Ar << Name;
-				Ar << Data;
-				ModelData.Add(Name, MakeShared<UE::NNE::FSharedModelData>(MakeSharedBufferFromArray(MoveTemp(Data)), 0));
+				Ar << Data_Size32Bits;
+				ModelData.Add(Name, MakeShared<UE::NNE::FSharedModelData>(MakeSharedBufferFromArray(MoveTemp(Data_Size32Bits)), 0));
 			}
 			UE_LOG(LogNNE, Warning, TEXT("[DEPRECATION] UNNEModelData: The asset %s (v0) is deprecated. Please right-click the asset and select 'Save' to update it to the latest version."), *this->GetName());
 			break;
@@ -254,14 +257,15 @@ void UNNEModelData::Serialize(FArchive& Ar)
 				Ar << AssetImportData;
 			}
 			Ar << FileType;
-			Ar << FileData;
+			Ar << Data_Size32Bits;
+			FileData = MoveTemp(Data_Size32Bits);
 			Ar << FileId;
 			Ar << NumItems;
 			for (Index = 0; Index < NumItems; Index++)
 			{
 				Ar << Name;
-				Ar << Data;
-				ModelData.Add(Name, MakeShared<UE::NNE::FSharedModelData>(MakeSharedBufferFromArray(MoveTemp(Data)), 0));
+				Ar << Data_Size32Bits;
+				ModelData.Add(Name, MakeShared<UE::NNE::FSharedModelData>(MakeSharedBufferFromArray(MoveTemp(Data_Size32Bits)), 0));
 			}
 			UE_LOG(LogNNE, Warning, TEXT("[DEPRECATION] UNNEModelData: The asset %s (v1) is deprecated. Please right-click the asset and select 'Save' to update it to the latest version."), *this->GetName());
 			break;
@@ -269,7 +273,8 @@ void UNNEModelData::Serialize(FArchive& Ar)
 		case UE::NNE::ModelData::Version::V2:
 			Ar << TargetRuntimes;
 			Ar << FileType;
-			Ar << FileData;
+			Ar << Data_Size32Bits;
+			FileData = MoveTemp(Data_Size32Bits);
 			Ar << FileId;
 			Ar << NumItems;
 			for (Index = 0; Index < NumItems; Index++)
@@ -284,6 +289,33 @@ void UNNEModelData::Serialize(FArchive& Ar)
 			break;
 
 		case UE::NNE::ModelData::Version::V3:
+			Ar << TargetRuntimes;
+			Ar << FileType;
+			Ar << Data_Size32Bits;
+			FileData = MoveTemp(Data_Size32Bits);
+			Ar << NumAdditionalFileDataItems;
+			AdditionalFileData.Empty();
+			for (Index = 0; Index < NumAdditionalFileDataItems; Index++)
+			{
+				Ar << Name;
+				Ar << Data_Size32Bits;
+				Data = MoveTemp(Data_Size32Bits);
+				AdditionalFileData.Add(Name, Data);
+			}
+			Ar << FileId;
+			Ar << NumItems;
+			for (Index = 0; Index < NumItems; Index++)
+			{
+				Ar << Name;
+				Ar << MemoryAlignment;
+				Ar << DataSize;
+				RawData = FMemory::Malloc(DataSize, MemoryAlignment);
+				Ar.Serialize(RawData, DataSize);
+				ModelData.Add(Name, MakeShared<UE::NNE::FSharedModelData>(FSharedBuffer::TakeOwnership(RawData, DataSize, FMemory::Free), MemoryAlignment));
+			}
+			break;
+
+		case UE::NNE::ModelData::Version::V4:
 			Ar << TargetRuntimes;
 			Ar << FileType;
 			Ar << FileData;
@@ -315,7 +347,7 @@ void UNNEModelData::Serialize(FArchive& Ar)
 	}
 }
 
-void UNNEModelData::Init(const FString& Type, TConstArrayView<uint8> Buffer, const TMap<FString, TConstArrayView<uint8>>& AdditionalBuffers)
+void UNNEModelData::Init(const FString& Type, TConstArrayView64<uint8> Buffer, const TMap<FString, TConstArrayView64<uint8>>& AdditionalBuffers)
 {
 	TargetRuntimes.Empty();
 	FileType = Type;
@@ -358,12 +390,12 @@ FString UNNEModelData::GetFileType() const
 	return FileType;
 }
 
-TConstArrayView<uint8> UNNEModelData::GetFileData() const
+TConstArrayView64<uint8> UNNEModelData::GetFileData() const
 {
 	return FileData;
 }
 
-TConstArrayView<uint8> UNNEModelData::GetAdditionalFileData(const FString& Key) const
+TConstArrayView64<uint8> UNNEModelData::GetAdditionalFileData(const FString& Key) const
 {
 	return AdditionalFileData.Contains(Key) ? AdditionalFileData[Key] : TConstArrayView<uint8>();
 }
@@ -408,7 +440,7 @@ TSharedPtr<UE::NNE::FSharedModelData> UNNEModelData::GetModelData(const FString&
 		return TSharedPtr<UE::NNE::FSharedModelData>();
 	}
 
-	TMap<FString, TConstArrayView<uint8>> AdditionalFileDataView;
+	TMap<FString, TConstArrayView64<uint8>> AdditionalFileDataView;
 	TSet<FString> Keys;
 	AdditionalFileData.GetKeys(Keys);
 	for (auto& Key : Keys)
diff --git a/Engine/Source/Runtime/NNE/Public/NNEModelData.h b/Engine/Source/Runtime/NNE/Public/NNEModelData.h
--- a/Engine/Source/Runtime/NNE/Public/NNEModelData.h
+++ b/Engine/Source/Runtime/NNE/Public/NNEModelData.h
@@ -39,7 +39,7 @@ namespace UE::NNE
 		 *
 		 * @return A const array view of the shared data.
 		 */
-		TConstArrayView<uint8> GetView() const;
+		TConstArrayView64<uint8> GetView() const;
 
 		/**
 		 * Get the memory alignment with which the data has been aligned.
@@ -90,7 +90,7 @@ public:
 	 * @param Buffer The raw binary file data of the imported model to be copied into this asset.
 	 * @param AdditionalBuffers Additional raw binary data of the model to be copied into this asset.
 	 */
-	void Init(const FString& Type, TConstArrayView<uint8> Buffer, const TMap<FString, TConstArrayView<uint8>>& AdditionalBuffers = TMap<FString, TConstArrayView<uint8>>());
+	void Init(const FString& Type, TConstArrayView64<uint8> Buffer, const TMap<FString, TConstArrayView64<uint8>>& AdditionalBuffers = TMap<FString, TConstArrayView64<uint8>>());
 
 	/**
 	 * In editor: Get the target runtimes this model data will be cooked for. An empty list means all runtimes.
@@ -125,7 +125,7 @@ public:
 	 *
 	 * @return The FileData.
 	 */
-	TConstArrayView<uint8> GetFileData() const;
+	TConstArrayView64<uint8> GetFileData() const;
 
 	/**
 	 * Get read only access to AdditionalFileData.
@@ -135,7 +135,7 @@ public:
 	 *
 	 * @return The AdditionalFileData with a given Key if it exists and an empty view in standalone or when the key does not exist.
 	 */
-	TConstArrayView<uint8> GetAdditionalFileData(const FString& Key) const;
+	TConstArrayView64<uint8> GetAdditionalFileData(const FString& Key) const;
 
 	/**
 	 * Clears the FileData and the FileType.
@@ -185,12 +185,12 @@ private:
 	/**
 	 * The raw binary file data of the imported model.
 	 */
-	TArray<uint8> FileData;
+	TArray64<uint8> FileData;
 
 	/**
 	 * Additional raw binary data of the imported model.
 	 */
-	TMap<FString, TArray<uint8>> AdditionalFileData;
+	TMap<FString, TArray64<uint8>> AdditionalFileData;
 
 	/**
 	 * A Guid that uniquely identifies this model. This is used to cache optimized models in the editor.
diff --git a/Engine/Source/Runtime/NNE/Public/NNERuntime.h b/Engine/Source/Runtime/NNE/Public/NNERuntime.h
--- a/Engine/Source/Runtime/NNE/Public/NNERuntime.h
+++ b/Engine/Source/Runtime/NNE/Public/NNERuntime.h
@@ -59,7 +59,7 @@ public:
 	 * @param TargetPlatform The Interface identifying the target platform for which the data needs to be created. A null pointer indicates the currently compiled/running platform.
 	 * @return True if the runtime is able to create model data, false otherwise.
 	 */
-	virtual ECanCreateModelDataStatus CanCreateModelData(const FString& FileType, TConstArrayView<uint8> FileData, const TMap<FString, TConstArrayView<uint8>>& AdditionalFileData, const FGuid& FileId, const ITargetPlatform* TargetPlatform) const = 0;
+	virtual ECanCreateModelDataStatus CanCreateModelData(const FString& FileType, TConstArrayView64<uint8> FileData, const TMap<FString, TConstArrayView64<uint8>>& AdditionalFileData, const FGuid& FileId, const ITargetPlatform* TargetPlatform) const = 0;
 
 	/**
 	 * Create model data given some raw file data.
@@ -71,7 +71,7 @@ public:
 	 * @param TargetPlatform The Interface identifying the target platform for which the data needs to be created. A null pointer indicates the currently compiled/running platform.
 	 * @return A pointer to data representing the runtime specific representation of the model to be stored by UNNEModelData on success or an invalid pointer otherwise.
 	 */
-	virtual TSharedPtr<UE::NNE::FSharedModelData> CreateModelData(const FString& FileType, TConstArrayView<uint8> FileData, const TMap<FString, TConstArrayView<uint8>>& AdditionalFileData, const FGuid& FileId, const ITargetPlatform* TargetPlatform) = 0;
+	virtual TSharedPtr<UE::NNE::FSharedModelData> CreateModelData(const FString& FileType, TConstArrayView64<uint8> FileData, const TMap<FString, TConstArrayView64<uint8>>& AdditionalFileData, const FGuid& FileId, const ITargetPlatform* TargetPlatform) = 0;
 
 	/**
 	 * Get an id uniquely identifying the model data. This is used by UNNEModelData to create a FCacheKey to get and add the model data from and to a DDC.
@@ -83,5 +83,5 @@ public:
 	 * @param TargetPlatform The Interface identifying the target platform for which the data is produced or by which it will be consumed. A null pointer indicates the currently compiled/running platform.
 	 * @return A string that identifies the model data uniquely, e.g. by containing the FileId, model data version and target platform display name.
 	 */
-	virtual FString GetModelDataIdentifier(const FString& FileType, TConstArrayView<uint8> FileData, const TMap<FString, TConstArrayView<uint8>>& AdditionalFileData, const FGuid& FileId, const ITargetPlatform* TargetPlatform) const = 0;
+	virtual FString GetModelDataIdentifier(const FString& FileType, TConstArrayView64<uint8> FileData, const TMap<FString, TConstArrayView64<uint8>>& AdditionalFileData, const FGuid& FileId, const ITargetPlatform* TargetPlatform) const = 0;
 };
\ No newline at end of file
