[NNE] NNERuntimeIREE use intermediate folder as cache.

#tests Editor, Win64 checked caching, staging and packaging

#rnx

diff --git a/Engine/Plugins/Experimental/NNERuntimeIREE/Source/NNERuntimeIREE/Private/NNERuntimeIREE.cpp b/Engine/Plugins/Experimental/NNERuntimeIREE/Source/NNERuntimeIREE/Private/NNERuntimeIREE.cpp
--- a/Engine/Plugins/Experimental/NNERuntimeIREE/Source/NNERuntimeIREE/Private/NNERuntimeIREE.cpp
+++ b/Engine/Plugins/Experimental/NNERuntimeIREE/Source/NNERuntimeIREE/Private/NNERuntimeIREE.cpp
@@ -24,11 +24,33 @@
 #include "NNERuntimeIREELog.h"
 #include "NNERuntimeIREEMetaData.h"
 #include "NNERuntimeIREEModel.h"
+#include "NNERuntimeIREEModelData.h"
 #include "Serialization/MemoryReader.h"
 #include "Serialization/MemoryWriter.h"
 
 namespace UE::NNERuntimeIREE::CPU::Private
 {
+	FString GetTargetPlatformName(const ITargetPlatform* TargetPlatform)
+	{
+		return TargetPlatform ? TargetPlatform->IniPlatformName() : UGameplayStatics::GetPlatformName();
+	}
+
+	FString GetBinariesSubdirectory(const FString& PlatformName)
+	{
+		if (PlatformName.Equals("Windows"))
+		{
+			if (PLATFORM_64BITS)
+			{
+				return TEXT("Win64");
+			}
+			return TEXT("Win32");
+		}
+		else
+		{
+			return PlatformName;
+		}
+	}
+
 	FString GetModelDataIdentifier(const FString& RuntimeName, const FGuid& Guid, const FString& FileIdString, const FString& PlatformName, const FString& Architecture)
 	{
 		return RuntimeName + "-" + Guid.ToString(EGuidFormats::Digits) + "-" + FString::FromInt(UNNERuntimeIREECpu::Version) + "-" + FileIdString + "-" + PlatformName + (!Architecture.IsEmpty() ? ("-" + Architecture) : "");
@@ -36,24 +58,31 @@ namespace UE::NNERuntimeIREE::CPU::Private
 
 	FString GetIntermediateModelDirPath(const FString& PlatformName, const FString& ModelName)
 	{
-		return FPaths::Combine("Intermediate", "Build", PlatformName, UE_PLUGIN_NAME, ModelName);
+		return FPaths::Combine("Intermediate", "Build", GetBinariesSubdirectory(PlatformName), UE_PLUGIN_NAME, ModelName);
 	}
 
 	FString GetStagedModelDirPath(const FString& PlatformName)
 	{
-		FString PlatformNameShort = PlatformName.Equals("Windows") ? "Win64" : PlatformName;
-
-		return FPaths::Combine("Binaries", PlatformNameShort, UE_PLUGIN_NAME);
+		return FPaths::Combine("Binaries", GetBinariesSubdirectory(PlatformName), UE_PLUGIN_NAME);
 	}
 
 	FString GetPackagedModelDirPath(const FString& PlatformName)
 	{
 		return GetStagedModelDirPath(PlatformName);
 	}
+
+	FString GetSharedLibDirPath(const FString& ModelName)
+	{
+#if WITH_EDITOR
+	return GetIntermediateModelDirPath(UGameplayStatics::GetPlatformName(), ModelName);
+#else
+	return GetPackagedModelDirPath(UGameplayStatics::GetPlatformName());
+#endif // WITH_EDITOR
+	}
 } // UE::NNERuntimeIREE::CPU::Private
 
 FGuid UNNERuntimeIREECpu::GUID = FGuid((int32)'I', (int32)'C', (int32)'P', (int32)'U');
-int32 UNNERuntimeIREECpu::Version = 0x00000004;
+int32 UNNERuntimeIREECpu::Version = 0x00000005;
 
 FString UNNERuntimeIREECpu::GetRuntimeName() const
 {
@@ -77,66 +106,128 @@ UNNERuntimeIREECpu::ECanCreateModelDataStatus UNNERuntimeIREECpu::CanCreateModel
 
 TSharedPtr<UE::NNE::FSharedModelData> UNNERuntimeIREECpu::CreateModelData(const FString& FileType, TConstArrayView64<uint8> FileData, const TMap<FString, TConstArrayView64<uint8>>& AdditionalFileData, const FGuid& FileId, const ITargetPlatform* TargetPlatform)
 {
+	SCOPED_NAMED_EVENT_TEXT("UNNERuntimeIREECpu::CreateModelData", FColor::Magenta);
+
 #if WITH_EDITOR
 	using namespace UE::NNERuntimeIREE::CPU::Private;
 
-	FString TargetPlatformName = TargetPlatform ? TargetPlatform->IniPlatformName() : UGameplayStatics::GetPlatformName();
+	const FString TargetPlatformName = GetTargetPlatformName(TargetPlatform);
 	if (CanCreateModelData(FileType, FileData, AdditionalFileData, FileId, TargetPlatform) != ECanCreateModelDataStatus::Ok)
 	{
 		UE_LOG(LogNNERuntimeIREE, Warning, TEXT("UNNERuntimeIREECpu cannot create the model data with id %s (Filetype: %s) for platform %s"), *FileId.ToString(EGuidFormats::Digits).ToLower(), *FileType, *TargetPlatformName);
 		return TSharedPtr<UE::NNE::FSharedModelData>();
 	}
 
-	TUniquePtr<UE::NNERuntimeIREE::CPU::FCompiler> Compiler = UE::NNERuntimeIREE::CPU::FCompiler::Make(TargetPlatformName);
-	if (!Compiler.IsValid())
-	{
-		UE_LOG(LogNNERuntimeIREE, Warning, TEXT("UNNERuntimeIREECpu failed to create a compiler to compile for platform %s"), *TargetPlatformName);
-		return TSharedPtr<UE::NNE::FSharedModelData>();
-	}
-	
-	FString FileIdString = FileId.ToString(EGuidFormats::Digits).ToLower();
-	FString IntermediateDir = FPaths::ConvertRelativePathToFull(FPaths::Combine(FPaths::ProjectDir(), GetIntermediateModelDirPath(TargetPlatformName, FileIdString)));
 	IPlatformFile& PlatformFile = FPlatformFileManager::Get().GetPlatformFile();
-	PlatformFile.DeleteDirectoryRecursively(*IntermediateDir);
-	PlatformFile.CreateDirectoryTree(*IntermediateDir);
 
-	TArray<UE::NNERuntimeIREE::CPU::FCompilerResult> CompilerResults;
-	UNNERuntimeIREEModuleMetaData* CompilerModuleMetaData = NewObject<UNNERuntimeIREEModuleMetaData>();
-	FString StagingDir = FPaths::ConvertRelativePathToFull(FPaths::Combine(FPaths::ProjectDir(), GetPackagedModelDirPath(TargetPlatformName)));
-	if (!Compiler->CompileMlir(FileData, FileIdString, IntermediateDir, StagingDir, CompilerResults, CompilerModuleMetaData))
-	{
-		UE_LOG(LogNNERuntimeIREE, Warning, TEXT("UNNERuntimeIREECpu failed to compile model %s"), *FileIdString);
-		return TSharedPtr<UE::NNE::FSharedModelData>();
-	}
+	const FString FileIdString = FileId.ToString(EGuidFormats::Digits).ToLower();
+	const FString IntermediateDirFullPath = FPaths::ConvertRelativePathToFull(FPaths::Combine(FPaths::ProjectDir(), GetIntermediateModelDirPath(TargetPlatformName, FileIdString)));
+	const FString SharedLibraryDirFullPath = FPaths::ConvertRelativePathToFull(FPaths::Combine(FPaths::ProjectDir(), GetSharedLibDirPath(FileIdString)));
+
+	FString IREEModelDataFilePath = FPaths::Combine(IntermediateDirFullPath, FileIdString) + ".ireemodeldata";
 
 	TArray64<uint8> ResultData;
-	FMemoryWriter64 Writer(ResultData);
-	Writer << UNNERuntimeIREECpu::GUID;
-	Writer << UNNERuntimeIREECpu::Version;
-	FGuid FileIdCopy = FileId;
-	Writer << FileIdCopy;
+	TWeakObjectPtr<UNNERuntimeIREEModelData> IREEModelData = NewObject<UNNERuntimeIREEModelData>();
+	FNNERuntimeIREECompilerResultCPU CompilerResult{};
 
-	TArray64<uint8> ModuleMetaData;
-	if (AdditionalFileData.Contains("IREEModuleMetaData"))
+	bool bNeedCompileMlir = true;
+	if (PlatformFile.FileExists(*IREEModelDataFilePath))
 	{
-		ModuleMetaData = AdditionalFileData["IREEModuleMetaData"];
+		SCOPED_NAMED_EVENT_TEXT("Validate", FColor::Magenta);
+
+		FFileHelper::LoadFileToArray(ResultData, *IREEModelDataFilePath);
+		
+		{
+			FMemoryReaderView Reader(ResultData);
+			IREEModelData->Serialize(Reader);
+		}
+
+		check(FileIdString.Equals(IREEModelData->FileId.ToString(EGuidFormats::Digits).ToLower()));
+
+		{
+			FMemoryReaderView Reader(IREEModelData->CompilerResult);
+			FNNERuntimeIREECompilerResultCPU::StaticStruct()->SerializeBin(Reader, &CompilerResult);
+		}
+
+		bNeedCompileMlir = false;
+		for (int32 i = 0; i < CompilerResult.ArchitectureInfos.Num(); i++)
+		{
+			const FNNERuntimeIREEArchitectureInfoCPU& Info = CompilerResult.ArchitectureInfos[i];
+			const FString SharedLibrarySubDirFullPath = FPaths::Combine(SharedLibraryDirFullPath, Info.RelativeDirPath);
+
+			const FString SharedLibraryFilePath = FPaths::Combine(SharedLibrarySubDirFullPath, Info.SharedLibraryFileName);
+			const FString VmfbFilePath = FPaths::Combine(SharedLibrarySubDirFullPath, Info.VmfbFileName);
+
+			bNeedCompileMlir |= !PlatformFile.FileExists(*SharedLibraryFilePath);
+			bNeedCompileMlir |= !PlatformFile.FileExists(*VmfbFilePath);
+		}
 	}
-	if (ModuleMetaData.IsEmpty())
+	
+	if (bNeedCompileMlir)
 	{
-		FMemoryWriter64 ObjectWriter(ModuleMetaData);
-		CompilerModuleMetaData->Serialize(ObjectWriter);
+		SCOPED_NAMED_EVENT_TEXT("Compile", FColor::Magenta);
+
+		PlatformFile.DeleteDirectoryRecursively(*IntermediateDirFullPath);
+		PlatformFile.CreateDirectoryTree(*IntermediateDirFullPath);
+
+		TUniquePtr<UE::NNERuntimeIREE::CPU::FCompiler> Compiler = UE::NNERuntimeIREE::CPU::FCompiler::Make(TargetPlatformName);
+		if (!Compiler.IsValid())
+		{
+			UE_LOG(LogNNERuntimeIREE, Warning, TEXT("UNNERuntimeIREECpu failed to create a compiler to compile for platform %s"), *TargetPlatformName);
+			return TSharedPtr<UE::NNE::FSharedModelData>();
+		}
+
+		TWeakObjectPtr<UNNERuntimeIREEModuleMetaData> CompilerModuleMetaData = NewObject<UNNERuntimeIREEModuleMetaData>();
+
+		if (!Compiler->CompileMlir(FileData, FileIdString, IntermediateDirFullPath, CompilerResult, *CompilerModuleMetaData))
+		{
+			UE_LOG(LogNNERuntimeIREE, Warning, TEXT("UNNERuntimeIREECpu failed to compile model %s"), *FileIdString);
+			return TSharedPtr<UE::NNE::FSharedModelData>();
+		}
+
+		IREEModelData->GUID = UNNERuntimeIREECpu::GUID;
+		IREEModelData->Version = UNNERuntimeIREECpu::Version;
+		IREEModelData->FileId = FileId;
+		if (AdditionalFileData.Contains("IREEModuleMetaData"))
+		{
+			IREEModelData->ModuleMetaData = AdditionalFileData["IREEModuleMetaData"];
+		}
+		if (IREEModelData->ModuleMetaData.IsEmpty())
+		{
+			FMemoryWriter64 Writer(IREEModelData->ModuleMetaData);
+			CompilerModuleMetaData->Serialize(Writer);
+		}
+		{
+			FMemoryWriter64 Writer(IREEModelData->CompilerResult);
+			FNNERuntimeIREECompilerResultCPU::StaticStruct()->SerializeBin(Writer, &CompilerResult);
+		}
+
+		{
+			FMemoryWriter64 Writer(ResultData);
+			IREEModelData->Serialize(Writer);
+		}
+
+		FFileHelper::SaveArrayToFile(ResultData, *IREEModelDataFilePath);
 	}
-	Writer << ModuleMetaData;
 
-	int32 NumArchitectures = CompilerResults.Num();
-	Writer << NumArchitectures;
-	for (int32 i = 0; i < NumArchitectures; i++)
+	// Copy files for staging
+	FString StagingDirFullPath = FPaths::ConvertRelativePathToFull(FPaths::Combine(FPaths::ProjectDir(), GetPackagedModelDirPath(TargetPlatformName)));
+	for (int32 i = 0; i < CompilerResult.ArchitectureInfos.Num(); i++)
 	{
-		Writer << CompilerResults[i].Architecture;
-		Writer << CompilerResults[i].RelativeDirPath;
-		Writer << CompilerResults[i].SharedLibraryFileName;
-		Writer << CompilerResults[i].VmfbFileName;
-		Writer << CompilerResults[i].SharedLibraryEntryPointName;
+		SCOPED_NAMED_EVENT_TEXT("Copy", FColor::Magenta);
+
+		const FNNERuntimeIREEArchitectureInfoCPU& Info = CompilerResult.ArchitectureInfos[i];
+		const FString SharedLibrarySubDirFullPath = FPaths::Combine(SharedLibraryDirFullPath, Info.RelativeDirPath);
+		const FString StagingSubDirFullPath = FPaths::Combine(StagingDirFullPath, Info.Architecture);
+		
+		const FString SharedLibraryFilePathSrc = FPaths::Combine(SharedLibrarySubDirFullPath, Info.SharedLibraryFileName);
+		const FString VmfbFilePathSrc = FPaths::Combine(SharedLibrarySubDirFullPath, Info.VmfbFileName);
+
+		const FString SharedLibraryFilePathDest = FPaths::Combine(StagingSubDirFullPath, Info.SharedLibraryFileName);
+		const FString VmfbFilePathDest = FPaths::Combine(StagingSubDirFullPath, Info.VmfbFileName);
+
+		IFileManager::Get().Copy(*SharedLibraryFilePathDest, *SharedLibraryFilePathSrc, bNeedCompileMlir);
+		IFileManager::Get().Copy(*VmfbFilePathDest, *VmfbFilePathSrc, bNeedCompileMlir);
 	}
 
 	return MakeShared<UE::NNE::FSharedModelData>(MakeSharedBufferFromArray(MoveTemp(ResultData)), 0);
@@ -148,7 +239,7 @@ TSharedPtr<UE::NNE::FSharedModelData> UNNERuntimeIREECpu::CreateModelData(const
 FString UNNERuntimeIREECpu::GetModelDataIdentifier(const FString& FileType, TConstArrayView64<uint8> FileData, const TMap<FString, TConstArrayView64<uint8>>& AdditionalFileData, const FGuid& FileId, const ITargetPlatform* TargetPlatform) const
 {
 	// Leave architecture blank as there is only one model data for all architectures of a given platform, only the vmfb and shared lib are different
-	FString PlatformName = TargetPlatform ? TargetPlatform->IniPlatformName() : UGameplayStatics::GetPlatformName();
+	FString PlatformName = UE::NNERuntimeIREE::CPU::Private::GetTargetPlatformName(TargetPlatform);
 	return UE::NNERuntimeIREE::CPU::Private::GetModelDataIdentifier(GetRuntimeName(), UNNERuntimeIREECpu::GUID, FileId.ToString(EGuidFormats::Digits), PlatformName, "");
 }
 
@@ -162,18 +253,12 @@ UNNERuntimeIREECpu::ECanCreateModelCPUStatus UNNERuntimeIREECpu::CanCreateModelC
 		return ECanCreateModelCPUStatus::Fail;
 	}
 
-	TConstArrayView<uint8> SharedDataView = SharedData->GetView();
-	int32 GuidSize = sizeof(UNNERuntimeIREECpu::GUID);
-	int32 VersionSize = sizeof(UNNERuntimeIREECpu::Version);
-	if (SharedDataView.Num() <= GuidSize + VersionSize)
+	if (!UNNERuntimeIREEModelData::IsSameGuidAndVersion(SharedData->GetView(), UNNERuntimeIREECpu::GUID, UNNERuntimeIREECpu::Version))
 	{
 		return ECanCreateModelCPUStatus::Fail;
 	}
 
-	bool bResult = FGenericPlatformMemory::Memcmp(&(SharedDataView[0]), &(UNNERuntimeIREECpu::GUID), GuidSize) == 0;
-	bResult &= FGenericPlatformMemory::Memcmp(&(SharedDataView[GuidSize]), &(UNNERuntimeIREECpu::Version), VersionSize) == 0;
-
-	return bResult ? ECanCreateModelCPUStatus::Ok : ECanCreateModelCPUStatus::Fail;
+	return ECanCreateModelCPUStatus::Ok;
 }
 
 TSharedPtr<UE::NNE::IModelCPU> UNNERuntimeIREECpu::CreateModelCPU(const TObjectPtr<UNNEModelData> ModelData)
@@ -197,86 +282,61 @@ TSharedPtr<UE::NNE::IModelCPU> UNNERuntimeIREECpu::CreateModelCPU(const TObjectP
 
 	TSharedPtr<UE::NNE::FSharedModelData> SharedData = ModelData->GetModelData(GetRuntimeName());
 	check(SharedData.IsValid());
+
 	TConstArrayView<uint8> SharedDataView = SharedData->GetView();
-	FMemoryReaderView Reader(SharedDataView);
-	FGuid DataGuid = FGuid();
-	Reader << DataGuid;
-	int32 VersionVersion = 0;
-	Reader << VersionVersion;
-	FGuid FileId = FGuid();
-	Reader << FileId;
-
-	TArray<uint8> ModuleDataArray;
-	Reader << ModuleDataArray;
-	if (ModuleDataArray.IsEmpty())
+
+	TWeakObjectPtr<UNNERuntimeIREEModelData> IREEModelData = NewObject<UNNERuntimeIREEModelData>();
+	{
+		FMemoryReaderView Reader(SharedDataView);
+		IREEModelData->Serialize(Reader);
+	}
+
+	if (IREEModelData->ModuleMetaData.IsEmpty())
 	{
 		UE_LOG(LogNNERuntimeIREE, Warning, TEXT("UNNERuntimeIREECpu failed to find any module meta data, please reimport the original model"));
 		return TSharedPtr<UE::NNE::IModelCPU>();
 	}
-	UNNERuntimeIREEModuleMetaData* ModuleMetaData = NewObject<UNNERuntimeIREEModuleMetaData>();
-	FMemoryReaderView ObjectReader(ModuleDataArray);
-	ModuleMetaData->Serialize(ObjectReader);
+
+	TWeakObjectPtr<UNNERuntimeIREEModuleMetaData> ModuleMetaData = NewObject<UNNERuntimeIREEModuleMetaData>();
+	{
+		FMemoryReaderView Reader(IREEModelData->ModuleMetaData);
+		ModuleMetaData->Serialize(Reader);
+	}
+
 	if (ModuleMetaData->FunctionMetaData.IsEmpty())
 	{
 		UE_LOG(LogNNERuntimeIREE, Warning, TEXT("UNNERuntimeIREECpu failed to parse the module meta data, please reimport the original model"));
 		return TSharedPtr<UE::NNE::IModelCPU>();
 	}
 
-	int32 NumArchitectures = 0;
-	Reader << NumArchitectures;
-
-	bool bFound = false;
-	FString Architecture = "";
-	FString RelativeDirPath = "";
-	FString SharedLibraryFileName = "";
-	FString VmfbFileName = "";
-	FString SharedLibraryEntryPointName = "";
-	for (int32 i = 0; i < NumArchitectures; i++)
-	{
-		FString TmpArchitecture = "";
-		Reader << TmpArchitecture;
-		FString TmpRelativeDirPath = "";
-		Reader << TmpRelativeDirPath;
-		FString TmpSharedLibraryFileName = "";
-		Reader << TmpSharedLibraryFileName;
-		FString TmpVmfbFileName = "";
-		Reader << TmpVmfbFileName;
-		FString TmpSharedLibraryEntryPointName = "";
-		Reader << TmpSharedLibraryEntryPointName;
-
-		if (TmpArchitecture.IsEmpty() && !bFound)
-		{
-			Architecture = TmpArchitecture;
-			RelativeDirPath = TmpRelativeDirPath;
-			SharedLibraryFileName = TmpSharedLibraryFileName;
-			VmfbFileName = TmpVmfbFileName;
-			SharedLibraryEntryPointName = TmpSharedLibraryEntryPointName;
-			bFound = true;
-		}
-		else if (TmpArchitecture.Equals(CurrentArchitecture))
+	FNNERuntimeIREECompilerResultCPU CompilerResult{};
+	{
+		FMemoryReaderView Reader(IREEModelData->CompilerResult);
+		FNNERuntimeIREECompilerResultCPU::StaticStruct()->SerializeBin(Reader, &CompilerResult);
+	}
+
+	int32 ArchitectureIndex = -1;
+	for (int32 i = 0; i < CompilerResult.ArchitectureInfos.Num(); i++)
+	{
+		if ((CompilerResult.ArchitectureInfos[i].Architecture.IsEmpty() && ArchitectureIndex < 0) || CompilerResult.ArchitectureInfos[i].Architecture.Equals(CurrentArchitecture))
 		{
-			Architecture = TmpArchitecture;
-			RelativeDirPath = TmpRelativeDirPath;
-			SharedLibraryFileName = TmpSharedLibraryFileName;
-			VmfbFileName = TmpVmfbFileName;
-			SharedLibraryEntryPointName = TmpSharedLibraryEntryPointName;
-			bFound = true;
+			ArchitectureIndex = i;
 		}
 	}
-	if (!bFound)
+
+	if (ArchitectureIndex < 0)
 	{
 		UE_LOG(LogNNERuntimeIREE, Warning, TEXT("UNNERuntimeIREECpu failed to find a matching architecture for \'%s\'"), *CurrentArchitecture);
 		return TSharedPtr<UE::NNE::IModelCPU>();
 	}
 
-	FString FileIdString = FileId.ToString(EGuidFormats::Digits).ToLower();
-#if WITH_EDITOR
-	FString SharedLibraryDirPath = FPaths::ConvertRelativePathToFull(FPaths::Combine(FPaths::ProjectDir(), GetIntermediateModelDirPath(UGameplayStatics::GetPlatformName(), FileIdString), RelativeDirPath));
-#else
-	FString SharedLibraryDirPath = FPaths::ConvertRelativePathToFull(FPaths::Combine(FPaths::ProjectDir(), GetPackagedModelDirPath(UGameplayStatics::GetPlatformName()), RelativeDirPath));
-#endif // WITH_EDITOR
+	const FNNERuntimeIREEArchitectureInfoCPU& ArchitectureInfo = CompilerResult.ArchitectureInfos[ArchitectureIndex];
+
+	const FString FileIdString = IREEModelData->FileId.ToString(EGuidFormats::Digits).ToLower();
+	const FString SharedLibraryDirFullPath = FPaths::ConvertRelativePathToFull(FPaths::Combine(FPaths::ProjectDir(), GetSharedLibDirPath(FileIdString)));
+	const FString SharedLibrarySubDirFullPath = FPaths::Combine(SharedLibraryDirFullPath, ArchitectureInfo.RelativeDirPath);
 
-	TSharedPtr<UE::NNE::IModelCPU> Model = UE::NNERuntimeIREE::CPU::FModel::Make(SharedLibraryDirPath, SharedLibraryFileName, VmfbFileName, SharedLibraryEntryPointName, *ModuleMetaData);
+	TSharedPtr<UE::NNE::IModelCPU> Model = UE::NNERuntimeIREE::CPU::FModel::Make(SharedLibrarySubDirFullPath, ArchitectureInfo.SharedLibraryFileName, ArchitectureInfo.VmfbFileName, ArchitectureInfo.SharedLibraryEntryPointName, *ModuleMetaData);
 	if (!Model.IsValid())
 	{
 		UE_LOG(LogNNERuntimeIREE, Warning, TEXT("UNNERuntimeIREECpu could not initialize the model created from model data with id %s"), *FileIdString);
@@ -323,7 +383,7 @@ TSharedPtr<UE::NNE::FSharedModelData> UNNERuntimeIREEGpu::CreateModelData(const
 
 FString UNNERuntimeIREEGpu::GetModelDataIdentifier(const FString& FileType, TConstArrayView64<uint8> FileData, const TMap<FString, TConstArrayView64<uint8>>& AdditionalFileData, const FGuid& FileId, const ITargetPlatform* TargetPlatform) const
 {
-	FString PlatformName = TargetPlatform ? TargetPlatform->IniPlatformName() : UGameplayStatics::GetPlatformName();
+	FString PlatformName = UE::NNERuntimeIREE::CPU::Private::GetTargetPlatformName(TargetPlatform);
 	return UE::NNERuntimeIREE::CPU::Private::GetModelDataIdentifier(GetRuntimeName(), GetGUID(), FileId.ToString(EGuidFormats::Digits), PlatformName, "");
 }
 
@@ -472,7 +532,7 @@ TSharedPtr<UE::NNE::FSharedModelData> UNNERuntimeIREERdg::CreateModelData(const
 
 FString UNNERuntimeIREERdg::GetModelDataIdentifier(const FString& FileType, TConstArrayView64<uint8> FileData, const TMap<FString, TConstArrayView64<uint8>>& AdditionalFileData, const FGuid& FileId, const ITargetPlatform* TargetPlatform) const
 {
-	FString PlatformName = TargetPlatform ? TargetPlatform->IniPlatformName() : UGameplayStatics::GetPlatformName();
+	FString PlatformName = UE::NNERuntimeIREE::CPU::Private::GetTargetPlatformName(TargetPlatform);
 	return UE::NNERuntimeIREE::CPU::Private::GetModelDataIdentifier(GetRuntimeName(), UNNERuntimeIREERdg::GUID, FileId.ToString(EGuidFormats::Digits), PlatformName, "");
 }
 
diff --git a/Engine/Plugins/Experimental/NNERuntimeIREE/Source/NNERuntimeIREE/Private/NNERuntimeIREECompiler.cpp b/Engine/Plugins/Experimental/NNERuntimeIREE/Source/NNERuntimeIREE/Private/NNERuntimeIREECompiler.cpp
--- a/Engine/Plugins/Experimental/NNERuntimeIREE/Source/NNERuntimeIREE/Private/NNERuntimeIREECompiler.cpp
+++ b/Engine/Plugins/Experimental/NNERuntimeIREE/Source/NNERuntimeIREE/Private/NNERuntimeIREECompiler.cpp
@@ -229,26 +229,34 @@ namespace UE::NNERuntimeIREE
 			return TUniquePtr<FCompiler>(new FCompiler(CompilerCommand, LinkerCommand, SharedLibExt, BuildTargets));
 		}
 
-		bool FCompiler::CompileMlir(TConstArrayView<uint8> InFileData, const FString& InModelName, const FString& InIntermediateDir, const FString& InStagingDir, TArray<FCompilerResult>& OutCompilerResults, UNNERuntimeIREEModuleMetaData* ModuleMetaData)
+		bool FCompiler::CompileMlir(TConstArrayView<uint8> InFileData, const FString& InModelName, const FString& InOutputDir, FNNERuntimeIREECompilerResultCPU& OutCompilerResult, UNNERuntimeIREEModuleMetaData& ModuleMetaData)
 		{
-			using namespace Private;
+			SCOPED_NAMED_EVENT_TEXT("FCompiler::CompileMlir", FColor::Magenta);
 
-			FString InputFilePath = FPaths::Combine(InIntermediateDir, InModelName) + ".mlir";
-			FFileHelper::SaveArrayToFile(InFileData, *InputFilePath);
+			using namespace Private;
 
-			if (ModuleMetaData)
 			{
+				SCOPED_NAMED_EVENT_TEXT("Metadata", FColor::Magenta);
+
 				FString FileDataString = "";
 				FileDataString.AppendChars((char*)InFileData.GetData(), InFileData.Num());
-				ModuleMetaData->ParseFromString(FileDataString);
+				ModuleMetaData.ParseFromString(FileDataString);
 			}
 
-			TArray<FCompilerResult> Results;
-			bool bResult = true;
 			IPlatformFile& PlatformFile = FPlatformFileManager::Get().GetPlatformFile();
+
+			FString InputFilePath = FPaths::Combine(InOutputDir, InModelName) + ".mlir";
+			if (!PlatformFile.FileExists(*InputFilePath))
+			{
+				SCOPED_NAMED_EVENT_TEXT("InputFile", FColor::Magenta);
+
+				FFileHelper::SaveArrayToFile(InFileData, *InputFilePath);
+			}
+
+			bool bResult = true;
 			for (int32 i = 0; i < BuildTargets.Num(); i++)
 			{
-				FString IntermediateDirPath = FPaths::Combine(InIntermediateDir, BuildTargets[i].Architecture);
+				FString IntermediateDirPath = FPaths::Combine(InOutputDir, BuildTargets[i].Architecture);
 				PlatformFile.CreateDirectoryTree(*IntermediateDirPath);
 				FString IntermediateFilePathNoExt = FPaths::Combine(IntermediateDirPath, InModelName);
 				FString ObjectFilePath = IntermediateFilePathNoExt + ".o";
@@ -265,7 +273,11 @@ namespace UE::NNERuntimeIREE
 				CompilerArguments.ReplaceInline(*FString("${VMFB_PATH}"), *(FString("\"") + VmfbFilePath + "\""));
 				CompilerArguments.ReplaceInline(*FString("${INPUT_PATH}"), *(FString("\"") + InputFilePath + "\""));
 
-				RunCommand(CompilerCommand, CompilerArguments, IntermediateFilePathNoExt + "_compile-log.txt");
+				{
+					SCOPED_NAMED_EVENT_TEXT("Compile", FColor::Magenta);
+
+					RunCommand(CompilerCommand, CompilerArguments, IntermediateFilePathNoExt + "_compile-log.txt");
+				}
 
 				if (!PlatformFile.FileExists(*ObjectFilePath) || !PlatformFile.FileExists(*VmfbFilePath))
 				{
@@ -285,7 +297,11 @@ namespace UE::NNERuntimeIREE
 				LinkerArguments.ReplaceInline(*FString("${OBJECT_PATH}"), *(FString("\"") + ObjectFilePath + "\""));
 				LinkerArguments.ReplaceInline(*FString("${SHARED_LIB_PATH}"), *(FString("\"") + SharedLibFilePath + "\""));
 
-				RunCommand(LinkerCommand, LinkerArguments, IntermediateFilePathNoExt + "_link-log.txt");
+				{
+					SCOPED_NAMED_EVENT_TEXT("Link", FColor::Magenta);
+
+					RunCommand(LinkerCommand, LinkerArguments, IntermediateFilePathNoExt + "_link-log.txt");
+				}
 
 				if (!PlatformFile.FileExists(*SharedLibFilePath))
 				{
@@ -295,20 +311,6 @@ namespace UE::NNERuntimeIREE
 					continue;
 				}
 
-				FString StagingFilePathNoExt = FPaths::Combine(InStagingDir, BuildTargets[i].Architecture, InModelName);
-				FString StagedVmfbFilePath = StagingFilePathNoExt + ".vmfb";
-				FString StagedSharedLibFilePath = StagingFilePathNoExt + SharedLibExt;
-				if (IFileManager::Get().Copy(*StagedSharedLibFilePath, *SharedLibFilePath) != COPY_OK)
-				{
-					UE_LOG(LogNNERuntimeIREE, Warning, TEXT("UNNERuntimeIREECpu failed to copy \"%s\" to \"%s\""), *SharedLibFilePath, *StagedSharedLibFilePath);
-					bResult = false;
-				}
-				if (IFileManager::Get().Copy(*StagedVmfbFilePath, *VmfbFilePath) != COPY_OK)
-				{
-					UE_LOG(LogNNERuntimeIREE, Warning, TEXT("UNNERuntimeIREECpu failed to copy \"%s\" to \"%s\""), *VmfbFilePath, *StagedVmfbFilePath);
-					bResult = false;
-				}
-
 				FString SharedLibraryEntryPointName = "";
 				FString HeaderPath = IntermediateFilePathNoExt + ".h";
 				if (!PlatformFile.FileExists(*HeaderPath))
@@ -332,20 +334,22 @@ namespace UE::NNERuntimeIREE
 					continue;
 				}
 
-				FCompilerResult Result;
-				Result.Architecture = BuildTargets[i].Architecture;
-				Result.RelativeDirPath = BuildTargets[i].Architecture;
-				Result.SharedLibraryFileName = InModelName + SharedLibExt;
-				Result.VmfbFileName = InModelName + ".vmfb";
-				Result.SharedLibraryEntryPointName = SharedLibraryEntryPointName;
-				Results.Add(Result);
+				FNNERuntimeIREEArchitectureInfoCPU ArchitectureInfo;
+				ArchitectureInfo.Architecture = BuildTargets[i].Architecture;
+				ArchitectureInfo.RelativeDirPath = BuildTargets[i].Architecture;
+				ArchitectureInfo.SharedLibraryFileName = InModelName + SharedLibExt;
+				ArchitectureInfo.VmfbFileName = InModelName + ".vmfb";
+				ArchitectureInfo.SharedLibraryEntryPointName = SharedLibraryEntryPointName;
+				OutCompilerResult.ArchitectureInfos.Add(MoveTemp(ArchitectureInfo));
 			}
 
-			bResult &= !Results.IsEmpty();
-			if (bResult)
+			bResult &= !OutCompilerResult.ArchitectureInfos.IsEmpty();
+
+			if (!bResult)
 			{
-				OutCompilerResults = Results;
+				OutCompilerResult.ArchitectureInfos.Empty();
 			}
+
 			return bResult;
 		}
 	} // CPU
diff --git a/Engine/Plugins/Experimental/NNERuntimeIREE/Source/NNERuntimeIREE/Private/NNERuntimeIREECompiler.h b/Engine/Plugins/Experimental/NNERuntimeIREE/Source/NNERuntimeIREE/Private/NNERuntimeIREECompiler.h
--- a/Engine/Plugins/Experimental/NNERuntimeIREE/Source/NNERuntimeIREE/Private/NNERuntimeIREECompiler.h
+++ b/Engine/Plugins/Experimental/NNERuntimeIREE/Source/NNERuntimeIREE/Private/NNERuntimeIREECompiler.h
@@ -2,13 +2,45 @@
 
 #pragma once
 
-#ifdef WITH_NNE_RUNTIME_IREE
-#if WITH_EDITOR
-
 #include "CoreMinimal.h"
 #include "NNERuntimeIREEMetaData.h"
 #include "Serialization/JsonSerializerMacros.h"
 
+#include "NNERuntimeIREECompiler.generated.h"
+
+USTRUCT()
+struct FNNERuntimeIREEArchitectureInfoCPU
+{
+	GENERATED_BODY()
+
+	UPROPERTY()
+	FString Architecture;
+
+	UPROPERTY()
+	FString RelativeDirPath;
+	
+	UPROPERTY()
+	FString SharedLibraryFileName;
+
+	UPROPERTY()
+	FString VmfbFileName;
+
+	UPROPERTY()
+	FString SharedLibraryEntryPointName;
+};
+
+USTRUCT()
+struct FNNERuntimeIREECompilerResultCPU
+{
+	GENERATED_BODY()
+
+	UPROPERTY()
+	TArray<FNNERuntimeIREEArchitectureInfoCPU> ArchitectureInfos;
+};
+
+#ifdef WITH_NNE_RUNTIME_IREE
+#if WITH_EDITOR
+
 namespace UE::NNERuntimeIREE::CPU
 {
 	struct FBuildTarget : public FJsonSerializable
@@ -39,15 +71,6 @@ namespace UE::NNERuntimeIREE::CPU
 		END_JSON_SERIALIZER
 	};
 
-	struct FCompilerResult
-	{
-		FString Architecture;
-		FString RelativeDirPath;
-		FString SharedLibraryFileName;
-		FString VmfbFileName;
-		FString SharedLibraryEntryPointName;
-	};
-
 	class FCompiler
 	{
 	private:
@@ -58,7 +81,7 @@ namespace UE::NNERuntimeIREE::CPU
 		static TUniquePtr<FCompiler> Make(const FString& InTargetPlatformName);
 
 	public:
-		bool CompileMlir(TConstArrayView<uint8> InFileData, const FString& InModelName, const FString& InIntermediateDir, const FString& InStagingDir, TArray<FCompilerResult>& OutCompilerResults, UNNERuntimeIREEModuleMetaData* ModuleMetaData);
+		bool CompileMlir(TConstArrayView<uint8> InFileData, const FString& InModelName, const FString& InOutputDir, FNNERuntimeIREECompilerResultCPU& OutCompilerResult, UNNERuntimeIREEModuleMetaData& ModuleMetaData);
 
 	private:
 		FString CompilerCommand;
diff --git a/Engine/Plugins/Experimental/NNERuntimeIREE/Source/NNERuntimeIREE/Private/NNERuntimeIREEModelData.cpp b/Engine/Plugins/Experimental/NNERuntimeIREE/Source/NNERuntimeIREE/Private/NNERuntimeIREEModelData.cpp
new file mode 100644
--- /dev/null
+++ b/Engine/Plugins/Experimental/NNERuntimeIREE/Source/NNERuntimeIREE/Private/NNERuntimeIREEModelData.cpp
@@ -0,0 +1,71 @@
+// Copyright Epic Games, Inc. All Rights Reserved.
+
+#include "NNERuntimeIREEModelData.h"
+
+#include "Containers/Array.h"
+#include "NNE.h"
+#include "NNERuntimeIREELog.h"
+#include "Serialization/CustomVersion.h"
+
+namespace UE::NNERuntimeIREE::ModelData::Private
+{
+
+enum Version : uint32
+{
+	V0 = 0, // Initial
+	// New versions can be added above this line
+	VersionPlusOne,
+	Latest = VersionPlusOne - 1
+};
+const FGuid GUID(0x6dcb835d, 0x9ac64a1d, 0x8165d871, 0x6122dab7);
+FCustomVersionRegistration Version(GUID, Version::Latest, TEXT("NNERuntimeIREEModelDataVersion"));// Always save with the latest version
+
+} // UE::NNERuntimeIREE::ModelData::Private
+
+void UNNERuntimeIREEModelData::Serialize(FArchive& Ar)
+{
+	// Store the asset version (no effect in load)
+	Ar.UsingCustomVersion(UE::NNERuntimeIREE::ModelData::Private::GUID);
+
+	if (Ar.IsSaving() || Ar.IsCountingMemory()) 
+	{
+		Ar << GUID;
+		Ar << Version;
+		Ar << FileId;
+		Ar << ModuleMetaData;
+		Ar << CompilerResult;
+	}
+	else
+	{
+		int32 NumArchitectures = 0;
+
+		switch (Ar.CustomVer(UE::NNERuntimeIREE::ModelData::Private::GUID))
+		{
+		case UE::NNERuntimeIREE::ModelData::Private::Version::V0:
+			Ar << GUID;
+			Ar << Version;
+			Ar << FileId;
+			Ar << ModuleMetaData;
+			Ar << CompilerResult;
+			break;
+		default:
+			UE_LOG(LogNNERuntimeIREE, Error, TEXT("UNNERuntimeIREEModelData: Unknown asset version %d: Deserialisation failed, please reimport the original model."), Ar.CustomVer(UE::NNERuntimeIREE::ModelData::Private::GUID));
+			break;
+		}
+	}
+}
+
+bool UNNERuntimeIREEModelData::IsSameGuidAndVersion(TConstArrayView64<uint8> Data, FGuid Guid, int32 Version)
+{
+	int32 GuidSize = sizeof(UNNERuntimeIREEModelData::GUID);
+	int32 VersionSize = sizeof(UNNERuntimeIREEModelData::Version);
+	if (Data.Num() <= GuidSize + VersionSize)
+	{
+		return false;
+	}
+
+	bool bResult = FGenericPlatformMemory::Memcmp(&(Data[0]), &Guid, GuidSize) == 0;
+	bResult &= FGenericPlatformMemory::Memcmp(&(Data[GuidSize]), &Version, VersionSize) == 0;
+
+	return bResult;
+}
\ No newline at end of file
diff --git a/Engine/Plugins/Experimental/NNERuntimeIREE/Source/NNERuntimeIREE/Private/NNERuntimeIREEModelData.h b/Engine/Plugins/Experimental/NNERuntimeIREE/Source/NNERuntimeIREE/Private/NNERuntimeIREEModelData.h
new file mode 100644
--- /dev/null
+++ b/Engine/Plugins/Experimental/NNERuntimeIREE/Source/NNERuntimeIREE/Private/NNERuntimeIREEModelData.h
@@ -0,0 +1,54 @@
+// Copyright Epic Games, Inc. All Rights Reserved.
+
+#pragma once
+
+#include "Containers/ContainersFwd.h"
+#include "NNERuntimeIREECompiler.h"
+#include "NNETypes.h"
+#include "UObject/Class.h"
+
+#include "NNERuntimeIREEModelData.generated.h"
+
+/**
+ * IREE model data class.
+ */
+UCLASS()
+class UNNERuntimeIREEModelData : public UObject
+{
+	GENERATED_BODY()
+
+public:
+	// UObject interface
+	virtual void Serialize(FArchive& Ar) override;
+	// End of UObject interface
+
+	/**
+	 * Check data for matching Guid and Version without deserializing everything.
+	 */
+	static bool IsSameGuidAndVersion(TConstArrayView64<uint8> Data, FGuid Guid, int32 Version);
+
+	/**
+	 * A Guid that uniquely identifies this IREE model data.
+	 */
+	FGuid GUID;
+
+	/**
+	 * Current version of this IREE model data.
+	 */
+	int32 Version;
+
+	/**
+	 * A Guid that uniquely identifies the model.
+	 */
+	FGuid FileId;
+
+	/**
+	 * Serialized module meta data.
+	 */
+	TArray64<uint8> ModuleMetaData;
+
+	/**
+	 * Serialized compiler output.
+	 */
+	TArray64<uint8> CompilerResult;
+};
