[NNE] NNERuntimIREE fixed model shared library path (wrong platform).

Problem:
IREE model shared lib does not get packaged for platforms other than Windows.

Solution:
Fix the model shared library path to target platform, instead of current, when packaging.

Tests:
Packaged on Windows for Windows and consoles and verified model shared lib gets packaged and can be used on target platform.

#tests Packaged for Windows and consoles and verified IREE model properly runs

#rnx

diff --git a/Engine/Plugins/Experimental/NNERuntimeIREE/Source/NNERuntimeIREE/Private/NNERuntimeIREE.cpp b/Engine/Plugins/Experimental/NNERuntimeIREE/Source/NNERuntimeIREE/Private/NNERuntimeIREE.cpp
--- a/Engine/Plugins/Experimental/NNERuntimeIREE/Source/NNERuntimeIREE/Private/NNERuntimeIREE.cpp
+++ b/Engine/Plugins/Experimental/NNERuntimeIREE/Source/NNERuntimeIREE/Private/NNERuntimeIREE.cpp
@@ -70,12 +70,12 @@ namespace UE::NNERuntimeIREE::CPU::Private
 		return GetStagedModelDirPath(PlatformName);
 	}
 
-	FString GetSharedLibDirPath(const FString& ModelName)
+	FString GetSharedLibDirPath(const FString& PlatformName, const FString& ModelName)
 	{
 #if WITH_EDITOR
-	return GetIntermediateModelDirPath(UGameplayStatics::GetPlatformName(), ModelName);
+	return GetIntermediateModelDirPath(PlatformName, ModelName);
 #else
-	return GetPackagedModelDirPath(UGameplayStatics::GetPlatformName());
+	return GetPackagedModelDirPath(PlatformName);
 #endif // WITH_EDITOR
 	}
 } // UE::NNERuntimeIREE::CPU::Private
@@ -121,7 +121,7 @@ TSharedPtr<UE::NNE::FSharedModelData> UNNERuntimeIREECpu::CreateModelData(const
 
 	const FString FileIdString = FileId.ToString(EGuidFormats::Digits).ToLower();
 	const FString IntermediateDirFullPath = FPaths::ConvertRelativePathToFull(FPaths::Combine(FPaths::ProjectDir(), GetIntermediateModelDirPath(TargetPlatformName, FileIdString)));
-	const FString SharedLibraryDirFullPath = FPaths::ConvertRelativePathToFull(FPaths::Combine(FPaths::ProjectDir(), GetSharedLibDirPath(FileIdString)));
+	const FString SharedLibraryDirFullPath = FPaths::ConvertRelativePathToFull(FPaths::Combine(FPaths::ProjectDir(), GetSharedLibDirPath(TargetPlatformName, FileIdString)));
 
 	FString IREEModelDataFilePath = FPaths::Combine(IntermediateDirFullPath, FileIdString) + ".ireemodeldata";
 
@@ -332,7 +332,7 @@ TSharedPtr<UE::NNE::IModelCPU> UNNERuntimeIREECpu::CreateModelCPU(const TObjectP
 	const FNNERuntimeIREEArchitectureInfoCPU& ArchitectureInfo = CompilerResult.ArchitectureInfos[ArchitectureIndex];
 
 	const FString FileIdString = IREEModelData->FileId.ToString(EGuidFormats::Digits).ToLower();
-	const FString SharedLibraryDirFullPath = FPaths::ConvertRelativePathToFull(FPaths::Combine(FPaths::ProjectDir(), GetSharedLibDirPath(FileIdString)));
+	const FString SharedLibraryDirFullPath = FPaths::ConvertRelativePathToFull(FPaths::Combine(FPaths::ProjectDir(), GetSharedLibDirPath(UGameplayStatics::GetPlatformName(), FileIdString)));
 	const FString SharedLibrarySubDirFullPath = FPaths::Combine(SharedLibraryDirFullPath, ArchitectureInfo.RelativeDirPath);
 
 	TSharedPtr<UE::NNE::IModelCPU> Model = UE::NNERuntimeIREE::CPU::FModel::Make(SharedLibrarySubDirFullPath, ArchitectureInfo.SharedLibraryFileName, ArchitectureInfo.VmfbFileName, ArchitectureInfo.SharedLibraryEntryPointName, *ModuleMetaData);
