[NNE] Use persistant memory serialization where needed.

Problem: Persistant memory serlialization must be used when a byte stream is to be serialized to disk.
Solution: Change where relevant and missing in NNE runtimes.

#test: NNE smoke tests: in win64 editor debug + win64 standalone dev.

#rnx

diff --git a/Engine/Plugins/Experimental/NNERuntimeIREE/Source/NNERuntimeIREE/Private/NNERuntimeIREE.cpp b/Engine/Plugins/Experimental/NNERuntimeIREE/Source/NNERuntimeIREE/Private/NNERuntimeIREE.cpp
--- a/Engine/Plugins/Experimental/NNERuntimeIREE/Source/NNERuntimeIREE/Private/NNERuntimeIREE.cpp
+++ b/Engine/Plugins/Experimental/NNERuntimeIREE/Source/NNERuntimeIREE/Private/NNERuntimeIREE.cpp
@@ -137,14 +137,14 @@ TSharedPtr<UE::NNE::FSharedModelData> UNNERuntimeIREECpu::CreateModelData(const
 		FFileHelper::LoadFileToArray(ResultData, *IREEModelDataFilePath);
 		
 		{
-			FMemoryReaderView Reader(ResultData);
+			FMemoryReaderView Reader(ResultData, /*bIsPersitent =*/ true);
 			IREEModelData->Serialize(Reader);
 		}
 
 		check(FileIdString.Equals(IREEModelData->FileId.ToString(EGuidFormats::Digits).ToLower()));
 
 		{
-			FMemoryReaderView Reader(IREEModelData->CompilerResult);
+			FMemoryReaderView Reader(IREEModelData->CompilerResult, /*bIsPersitent =*/ true);
 			FNNERuntimeIREECompilerResultCPU::StaticStruct()->SerializeBin(Reader, &CompilerResult);
 		}
 
@@ -193,16 +193,16 @@ TSharedPtr<UE::NNE::FSharedModelData> UNNERuntimeIREECpu::CreateModelData(const
 		}
 		if (IREEModelData->ModuleMetaData.IsEmpty())
 		{
-			FMemoryWriter64 Writer(IREEModelData->ModuleMetaData);
+			FMemoryWriter64 Writer(IREEModelData->ModuleMetaData, /*bIsPersitent =*/ true);
 			CompilerModuleMetaData->Serialize(Writer);
 		}
 		{
-			FMemoryWriter64 Writer(IREEModelData->CompilerResult);
+			FMemoryWriter64 Writer(IREEModelData->CompilerResult, /*bIsPersitent =*/ true);
 			FNNERuntimeIREECompilerResultCPU::StaticStruct()->SerializeBin(Writer, &CompilerResult);
 		}
 
 		{
-			FMemoryWriter64 Writer(ResultData);
+			FMemoryWriter64 Writer(ResultData, /*bIsPersitent =*/ true);
 			IREEModelData->Serialize(Writer);
 		}
 
@@ -286,7 +286,7 @@ TSharedPtr<UE::NNE::IModelCPU> UNNERuntimeIREECpu::CreateModelCPU(const TObjectP
 
 	TWeakObjectPtr<UNNERuntimeIREEModelData> IREEModelData = NewObject<UNNERuntimeIREEModelData>();
 	{
-		FMemoryReaderView Reader(SharedDataView);
+		FMemoryReaderView Reader(SharedDataView, /*bIsPersitent =*/ true);
 		IREEModelData->Serialize(Reader);
 	}
 
@@ -298,7 +298,7 @@ TSharedPtr<UE::NNE::IModelCPU> UNNERuntimeIREECpu::CreateModelCPU(const TObjectP
 
 	TWeakObjectPtr<UNNERuntimeIREEModuleMetaData> ModuleMetaData = NewObject<UNNERuntimeIREEModuleMetaData>();
 	{
-		FMemoryReaderView Reader(IREEModelData->ModuleMetaData);
+		FMemoryReaderView Reader(IREEModelData->ModuleMetaData, /*bIsPersitent =*/ true);
 		ModuleMetaData->Serialize(Reader);
 	}
 
@@ -310,7 +310,7 @@ TSharedPtr<UE::NNE::IModelCPU> UNNERuntimeIREECpu::CreateModelCPU(const TObjectP
 
 	FNNERuntimeIREECompilerResultCPU CompilerResult{};
 	{
-		FMemoryReaderView Reader(IREEModelData->CompilerResult);
+		FMemoryReaderView Reader(IREEModelData->CompilerResult, /*bIsPersitent =*/ true);
 		FNNERuntimeIREECompilerResultCPU::StaticStruct()->SerializeBin(Reader, &CompilerResult);
 	}
 
diff --git a/Engine/Plugins/Experimental/NNERuntimeIREE/Source/NNERuntimeIREEEditor/Private/NNERuntimeIREEModelDataFactory.cpp b/Engine/Plugins/Experimental/NNERuntimeIREE/Source/NNERuntimeIREEEditor/Private/NNERuntimeIREEModelDataFactory.cpp
--- a/Engine/Plugins/Experimental/NNERuntimeIREE/Source/NNERuntimeIREEEditor/Private/NNERuntimeIREEModelDataFactory.cpp
+++ b/Engine/Plugins/Experimental/NNERuntimeIREE/Source/NNERuntimeIREEEditor/Private/NNERuntimeIREEModelDataFactory.cpp
@@ -203,7 +203,7 @@ UObject* UNNERuntimeIREEModelDataFactory::FactoryCreateBinary(UClass* Class, UOb
 	}
 
 	TArray<uint8> MetaDataByteArray;
-	FMemoryWriter Writer(MetaDataByteArray);
+	FMemoryWriter Writer(MetaDataByteArray, /*bIsPersitent =*/ true);
 	ModuleMetaData->Serialize(Writer);
 
 	TMap<FString, TConstArrayView64<uint8>> AdditionalFileData;
diff --git a/Engine/Plugins/Experimental/NNERuntimeRDG/Source/NNERuntimeRDG/Private/NNERuntimeRDGBase.h b/Engine/Plugins/Experimental/NNERuntimeRDG/Source/NNERuntimeRDG/Private/NNERuntimeRDGBase.h
--- a/Engine/Plugins/Experimental/NNERuntimeRDG/Source/NNERuntimeRDG/Private/NNERuntimeRDGBase.h
+++ b/Engine/Plugins/Experimental/NNERuntimeRDG/Source/NNERuntimeRDG/Private/NNERuntimeRDGBase.h
@@ -362,7 +362,7 @@ public:
 			return false;
 		}
 
-		FMemoryReader Reader(InputModel.Data);
+		FMemoryReader Reader(InputModel.Data, /*bIsPersitent =*/ true);
 		Reader.SetIsPersistent(true);
 
 		Format.Serialize(Reader);
diff --git a/Engine/Plugins/Experimental/NNERuntimeRDG/Source/NNERuntimeRDG/Private/NNERuntimeRDGModel.cpp b/Engine/Plugins/Experimental/NNERuntimeRDG/Source/NNERuntimeRDG/Private/NNERuntimeRDGModel.cpp
--- a/Engine/Plugins/Experimental/NNERuntimeRDG/Source/NNERuntimeRDG/Private/NNERuntimeRDGModel.cpp
+++ b/Engine/Plugins/Experimental/NNERuntimeRDG/Source/NNERuntimeRDG/Private/NNERuntimeRDGModel.cpp
@@ -17,7 +17,7 @@ bool FModelInstanceRDG::LoadModel(TConstArrayView<uint8> ModelData, FNNERuntimeF
 {
 	TConstArrayView<uint8> ModelBuffer = { &(ModelData.GetData()[GuidAndVersionSize]), ModelData.Num() - GuidAndVersionSize };
 
-	FMemoryReaderView Reader(ModelBuffer);
+	FMemoryReaderView Reader(ModelBuffer, /*bIsPersitent =*/ true);
 	Reader.SetIsPersistent(true);
 
 	Format.Serialize(Reader);
diff --git a/Engine/Plugins/Experimental/NNERuntimeRDG/Source/NNERuntimeRDGUtils/Private/NNERuntimeRDGUtilsModelBuilderNNE.cpp b/Engine/Plugins/Experimental/NNERuntimeRDG/Source/NNERuntimeRDGUtils/Private/NNERuntimeRDGUtilsModelBuilderNNE.cpp
--- a/Engine/Plugins/Experimental/NNERuntimeRDG/Source/NNERuntimeRDGUtils/Private/NNERuntimeRDGUtilsModelBuilderNNE.cpp
+++ b/Engine/Plugins/Experimental/NNERuntimeRDG/Source/NNERuntimeRDGUtils/Private/NNERuntimeRDGUtilsModelBuilderNNE.cpp
@@ -44,7 +44,7 @@ public:
 
 	virtual bool End(TArray<uint8>& Data) override
 	{
-		FMemoryWriter Writer(Data, true);
+		FMemoryWriter Writer(Data, /*bIsPersitent =*/ true);
 		
 		Format.Serialize(Writer);
 
diff --git a/Engine/Source/Runtime/NNE/Internal/NNEAttributeValue.h b/Engine/Source/Runtime/NNE/Internal/NNEAttributeValue.h
--- a/Engine/Source/Runtime/NNE/Internal/NNEAttributeValue.h
+++ b/Engine/Source/Runtime/NNE/Internal/NNEAttributeValue.h
@@ -22,7 +22,7 @@ struct FNNEAttributeValue
 	template<typename T>
 	explicit FNNEAttributeValue(T InValue)
 	{
-		FMemoryWriter writer(Value, true);
+		FMemoryWriter writer(Value, /*bIsPersitent =*/ true);
 		writer << InValue;
 
 		Type = TNNEAttributeValueTraits<T>::GetType();
@@ -35,7 +35,7 @@ struct FNNEAttributeValue
 
 		T Result;
 
-		FMemoryReader Reader(Value, true);
+		FMemoryReader Reader(Value, /*bIsPersitent =*/ true);
 		Reader << Result;
 
 		return Result;
