[NNE] NNERuntimeIREE Cpu standardize API surface to other runtimes behavior.

Problem: IREE model API surface is not standard with other plugin
Solution: Make it follow NNE team runtime convention

Test :
* Win64 Editor debug : NNE standard API unit test (from a different CL) + NNEQA IREE test from session frontend + NNEQA smoketests (including IREE test) + IREE plugin test
* Win64 Standalone : NNE standard API unit test (from a different CL) + NNEQA IREE test + NNEQA smoketests (including IREE test) + style transfer

#rnx

diff --git a/Engine/Plugins/Experimental/NNERuntimeIREE/Source/NNERuntimeIREE/Private/NNERuntimeIREEModel.cpp b/Engine/Plugins/Experimental/NNERuntimeIREE/Source/NNERuntimeIREE/Private/NNERuntimeIREEModel.cpp
--- a/Engine/Plugins/Experimental/NNERuntimeIREE/Source/NNERuntimeIREE/Private/NNERuntimeIREEModel.cpp
+++ b/Engine/Plugins/Experimental/NNERuntimeIREE/Source/NNERuntimeIREE/Private/NNERuntimeIREEModel.cpp
@@ -30,7 +30,7 @@ THIRD_PARTY_INCLUDES_END
 #if PLATFORM_MICROSOFT
 #include "Microsoft/HideMicrosoftPlatformAtomics.h"
 #include "Microsoft/HideMicrosoftPlatformTypes.h"
-#endif // PLATFORM_MICROSOF
+#endif // PLATFORM_MICROSOFT
 
 namespace UE::NNERuntimeIREE
 {
@@ -150,7 +150,7 @@ namespace UE::NNERuntimeIREE
 				Status = iree_runtime_instance_create(&InstanceOptions, iree_allocator_system(), &Instance);
 				if (!iree_status_is_ok(Status))
 				{
-					Private::PrintIREEError("UE::NNERuntimeIREE::Private::FInstance failed to create the instance", Status);
+					Private::PrintIREEError("CPU instance: Failed to create the instance", Status);
 
 					if (Instance)
 					{
@@ -184,7 +184,7 @@ namespace UE::NNERuntimeIREE
 				Status = iree_vm_bytecode_module_create(iree_runtime_instance_vm_instance(Instance), ModuleData, iree_allocator_null(), GetHostAllocator(), &Module);
 				if (!iree_status_is_ok(Status))
 				{
-					Private::PrintIREEError("UE::NNERuntimeIREE::Private::FInstance failed to create the module", Status);
+					Private::PrintIREEError("CPU instance: Failed to create the module", Status);
 
 					if (Module)
 					{
@@ -218,7 +218,7 @@ namespace UE::NNERuntimeIREE
 				Status = iree_hal_static_library_loader_create(IREE_ARRAYSIZE(LibraryList), LibraryList, iree_hal_executable_import_provider_null(), HostAllocator, &LibraryLoader);
 				if (!iree_status_is_ok(Status))
 				{
-					Private::PrintIREEError("UE::NNERuntimeIREE::Private::FInstance failed to create the library loader", Status);
+					Private::PrintIREEError("CPU instance: Failed to create the library loader", Status);
 
 					if (LibraryLoader)
 					{
@@ -234,7 +234,7 @@ namespace UE::NNERuntimeIREE
 				Status = iree_hal_allocator_create_heap(Identifier, HostAllocator, HostAllocator, &DeviceAllocator);
 				if (!iree_status_is_ok(Status))
 				{
-					Private::PrintIREEError("UE::NNERuntimeIREE::Private::FInstance failed to create the device allocator", Status);
+					Private::PrintIREEError("CPU instance: Failed to create the device allocator", Status);
 
 					if (DeviceAllocator)
 					{
@@ -252,7 +252,7 @@ namespace UE::NNERuntimeIREE
 				Status = iree_hal_sync_device_create(Identifier, &DeviceParams, 1, &LibraryLoader, DeviceAllocator, HostAllocator, &Device);
 				if (!iree_status_is_ok(Status))
 				{
-					Private::PrintIREEError("UE::NNERuntimeIREE::Private::FInstance failed to create the device", Status);
+					Private::PrintIREEError("CPU instance: Failed to create the device", Status);
 
 					if (Device)
 					{
@@ -291,7 +291,7 @@ namespace UE::NNERuntimeIREE
 				Status = iree_runtime_session_create_with_device(Instance, &SessionOptions, InDevice, GetHostAllocator(), &Session);
 				if (!iree_status_is_ok(Status))
 				{
-					Private::PrintIREEError("UE::NNERuntimeIREE::Private::FInstance failed to create the session", Status);
+					Private::PrintIREEError("CPU instance: Failed to create the session", Status);
 
 					if (Session)
 					{
@@ -354,13 +354,13 @@ namespace UE::NNERuntimeIREE
 				TUniquePtr<FArchive> Reader = TUniquePtr<FArchive>(IFileManager::Get().CreateFileReader(*CombinedPath, 0));
 				if (!Reader)
 				{
-					UE_LOG(LogNNERuntimeIREE, Error, TEXT("UE::NNERuntimeIREE::Private::FModule failed to open the vmfb data file '%s'"), *CombinedPath);
+					UE_LOG(LogNNERuntimeIREE, Error, TEXT("CPU module: Failed to open the vmfb data file '%s'"), *CombinedPath);
 					return TSharedPtr<FModule>();
 				}
 				int64 DataSize = Reader->TotalSize();
 				if (DataSize < 1)
 				{
-					UE_LOG(LogNNERuntimeIREE, Error, TEXT("UE::NNERuntimeIREE::Private::FModule's vmfb data file '%s' is empty"), *CombinedPath);
+					UE_LOG(LogNNERuntimeIREE, Error, TEXT("CPU module: vmfb data file '%s' is empty"), *CombinedPath);
 					return TSharedPtr<FModule>();
 				}
 
@@ -399,7 +399,7 @@ namespace UE::NNERuntimeIREE
 				Status = iree_runtime_session_append_module(InSession, Module);
 				if (!iree_status_is_ok(Status))
 				{
-					PrintIREEError("UE::NNERuntimeIREE::Private::FModule failed to append the module to the session", Status);
+					PrintIREEError("CPU module: Failed to append the module to the session", Status);
 					iree_status_free(Status);
 					return false;
 				}
@@ -439,7 +439,7 @@ namespace UE::NNERuntimeIREE
 
 				if (!bFound)
 				{
-					UE_LOG(LogNNERuntimeIREE, Error, TEXT("UE::NNERuntimeIREE::Private::FModule failed to find the module function %s"), *InFunctionName);
+					UE_LOG(LogNNERuntimeIREE, Error, TEXT("CPU module: Failed to find the module function %s"), *InFunctionName);
 					iree_status_free(Status);
 					return false;
 				}
@@ -487,7 +487,7 @@ namespace UE::NNERuntimeIREE
 					void* Library = FPlatformProcess::GetDllHandle(*CombinedPath);
 					if (!Library)
 					{
-						UE_LOG(LogNNERuntimeIREE, Error, TEXT("UE::NNERuntimeIREE::CPU::Private::FLibrary failed to load the shared library '%s'"), *CombinedPath);
+						UE_LOG(LogNNERuntimeIREE, Error, TEXT("CPU library: Failed to load the shared library '%s'"), *CombinedPath);
 						return TSharedPtr<FLibrary>();
 					}
 #else
@@ -496,7 +496,7 @@ namespace UE::NNERuntimeIREE
 					FPlatformProcess::PopDllDirectory(*InLibraryPath);
 					if (!Library)
 					{
-						UE_LOG(LogNNERuntimeIREE, Error, TEXT("UE::NNERuntimeIREE::CPU::Private::FLibrary failed to load the shared library '%s' from '%s'"), *InLibraryName, *InLibraryPath);
+						UE_LOG(LogNNERuntimeIREE, Error, TEXT("CPU library: Failed to load the shared library '%s' from '%s'"), *InLibraryName, *InLibraryPath);
 						return TSharedPtr<FLibrary>();
 					}
 #endif
@@ -514,7 +514,7 @@ namespace UE::NNERuntimeIREE
 						*OutFunctionPointer = Result;
 						return true;
 					}
-					UE_LOG(LogNNERuntimeIREE, Error, TEXT("UE::NNERuntimeIREE::CPU::Private::FLibrary failed to get the function %s"), *InFunctionName);
+					UE_LOG(LogNNERuntimeIREE, Error, TEXT("CPU library: Failed to get the function %s"), *InFunctionName);
 					return false;
 				}
 			};
@@ -626,34 +626,7 @@ namespace UE::NNERuntimeIREE
 					: Device(InDevice), Module(InModule), Session(InSession), Call(InCall), InputTensorDescs(InInputTensorDescs), OutputTensorDescs(InOutputTensorDescs)
 				{
 					check(InSession);
-
 					check(!InputTensorDescs.IsEmpty());
-					bool bAllConcrete = true;
-					for (int32 i = 0; i < InInputTensorDescs.Num(); i++)
-					{
-						bAllConcrete &= InInputTensorDescs[i].GetShape().IsConcrete();
-					}
-					if (bAllConcrete)
-					{
-						for (int32 i = 0; i < InInputTensorDescs.Num(); i++)
-						{
-							InputTensorShapes.Add(UE::NNE::FTensorShape::MakeFromSymbolic(InInputTensorDescs[i].GetShape()));
-						}
-					}
-
-					check(!OutputTensorDescs.IsEmpty());
-					bAllConcrete = true;
-					for (int32 i = 0; i < InOutputTensorDescs.Num(); i++)
-					{
-						bAllConcrete &= InOutputTensorDescs[i].GetShape().IsConcrete();
-					}
-					if (bAllConcrete)
-					{
-						for (int32 i = 0; i < InOutputTensorDescs.Num(); i++)
-						{
-							OutputTensorShapes.Add(UE::NNE::FTensorShape::MakeFromSymbolic(InOutputTensorDescs[i].GetShape()));
-						}
-					}
 				}
 
 			public:
@@ -702,7 +675,7 @@ namespace UE::NNERuntimeIREE
 					TConstArrayView<UE::NNE::FTensorDesc> OutputTensorDescs = InModule->GetFunctionMetaDataView()[0].OutputDescs;
 					if (!iree_status_is_ok(Status) || NumInputs != InputTensorDescs.Num() || NumOutputs != OutputTensorDescs.Num())
 					{
-						UE_LOG(LogNNERuntimeIREE, Error, TEXT("UE::NNERuntimeIREE::CPU::Private::FSession has a function signature mismatch in function %s"), *MainFunctionName);
+						UE_LOG(LogNNERuntimeIREE, Error, TEXT("CPU session: Function signature mismatch in function %s"), *MainFunctionName);
 						iree_runtime_session_release(Session);
 						iree_status_free(Status);
 						return TSharedPtr<FSession>();
@@ -712,7 +685,7 @@ namespace UE::NNERuntimeIREE
 					Status = iree_runtime_call_initialize(Session, MainFunction, &Call);
 					if (!iree_status_is_ok(Status))
 					{
-						UE::NNERuntimeIREE::Private::PrintIREEError("UE::NNERuntimeIREE::CPU::Private::FSession failed to initialize the session call", Status);
+						UE::NNERuntimeIREE::Private::PrintIREEError("CPU session: Failed to initialize the session call", Status);
 						iree_runtime_session_release(Session);
 						iree_status_free(Status);
 						return TSharedPtr<FSession>();
@@ -745,15 +718,64 @@ namespace UE::NNERuntimeIREE
 
 				ESetInputTensorShapesStatus SetInputTensorShapes(TConstArrayView<UE::NNE::FTensorShape> InInputShapes)
 				{
-					check(InputTensorDescs.Num() == InInputShapes.Num());
-					checkCode(for (int32 i = 0; i < InputTensorDescs.Num(); i++) { check(InInputShapes[i].IsCompatibleWith(InputTensorDescs[i].GetShape())); });
+					// OutputTensorShapes will be made available only if all shapes are concretes.
+					OutputTensorShapes.Reset();
+					bool bAllOutputShapeAreConcrete = true;
+					for (int32 i = 0; i < OutputTensorDescs.Num(); i++)
+					{
+						bAllOutputShapeAreConcrete &= OutputTensorDescs[i].GetShape().IsConcrete();
+					}
+					if (bAllOutputShapeAreConcrete)
+					{
+						for (int32 i = 0; i < OutputTensorDescs.Num(); i++)
+						{
+							OutputTensorShapes.Add(UE::NNE::FTensorShape::MakeFromSymbolic(OutputTensorDescs[i].GetShape()));
+						}
+					}
+
+					InputTensorShapes.Reset(InInputShapes.Num());
+					if (InInputShapes.Num() != InputTensorDescs.Num())
+					{
+						UE_LOG(LogNNERuntimeIREE, Error, TEXT("CPU session: Number of input shapes does not match number of input tensors"));
+						return ESetInputTensorShapesStatus::Fail;
+					}
+
+					for (int32 i = 0; i < InInputShapes.Num(); ++i)
+					{
+						const UE::NNE::FTensorDesc SymbolicDesc = InputTensorDescs[i];
+						if (!InInputShapes[i].IsCompatibleWith(SymbolicDesc.GetShape()))
+						{
+							UE_LOG(LogNNERuntimeIREE, Error, TEXT("CPU session: Input shape does not match input tensor %s of index %d"), *SymbolicDesc.GetName(), i);
+							return ESetInputTensorShapesStatus::Fail;
+						}
+					}
 					InputTensorShapes = InInputShapes;
+
 					return ESetInputTensorShapesStatus::Ok;
 				}
 
 				ERunSyncStatus RunSyncCPU(TConstArrayView<UE::NNE::FTensorBindingCPU> InInputBindings, TConstArrayView<UE::NNE::FTensorBindingCPU> InOutputBindings)
 				{
-					check(InInputBindings.Num() == InputTensorShapes.Num());
+					SCOPED_NAMED_EVENT_TEXT("NNERuntimeIREE::CPU::RunSync", FColor::Magenta);
+
+					// Verify the model inputs were prepared
+					if (InputTensorShapes.IsEmpty())
+					{
+						UE_LOG(LogNNERuntimeIREE, Error, TEXT("CPU session: Input shapes are not set, please call SetInputTensorShapes."));
+						return ERunSyncStatus::Fail;
+					}
+
+					if (InInputBindings.Num() != InputTensorShapes.Num())
+					{
+						UE_LOG(LogNNERuntimeIREE, Error, TEXT("CPU session: Input bindings need to match input tensor descriptor count (got %d, expected %d)."), InInputBindings.Num(), InputTensorShapes.Num());
+						return ERunSyncStatus::Fail;
+					}
+
+					if (!InOutputBindings.IsEmpty() && InOutputBindings.Num() != OutputTensorShapes.Num())
+					{
+						UE_LOG(LogNNERuntimeIREE, Error, TEXT("CPU session: Output binding can be empty or needs to match output tensor descriptor count (got %d, expected %d)."), InOutputBindings.Num(), OutputTensorShapes.Num());
+						return ERunSyncStatus::Fail;
+					}
 
 					iree_status_t Status = iree_ok_status();
 					check(iree_status_is_ok(Status));
@@ -762,8 +784,23 @@ namespace UE::NNERuntimeIREE
 
 					for (int32 i = 0; i < InInputBindings.Num(); i++)
 					{
-						check(InInputBindings[i].SizeInBytes >= InputTensorShapes[i].Volume() * InputTensorDescs[i].GetElementByteSize());
-						checkf(FMath::Modulo<uint64>((uint64)InInputBindings[i].Data, IREE_HAL_HEAP_BUFFER_ALIGNMENT) == 0, TEXT("NNERuntimeIREECpu requires input- and output-buffer memory to be aligned with %d bytes"), IREE_HAL_HEAP_BUFFER_ALIGNMENT);
+						if (!InInputBindings[i].Data && InInputBindings[i].SizeInBytes != 0)
+						{
+							UE_LOG(LogNNERuntimeIREE, Error, TEXT("CPU session: Binding input tensor %d is not set but given size is non-zero %d."), i, InInputBindings[i].SizeInBytes);
+							return ERunSyncStatus::Fail;
+						}
+
+						if (InInputBindings[i].SizeInBytes != (InputTensorShapes[i].Volume() * InputTensorDescs[i].GetElementByteSize()))
+						{
+							UE_LOG(LogNNERuntimeIREE, Error, TEXT("CPU session: Binding input tensor %d size does not match size given by tensor descriptor (got %d, expected %d)."), i, InInputBindings[i].SizeInBytes, InputTensorShapes[i].Volume() * InputTensorDescs[i].GetElementByteSize());
+							return ERunSyncStatus::Fail;
+						}
+
+						if (FMath::Modulo<uint64>((uint64)InInputBindings[i].Data, IREE_HAL_HEAP_BUFFER_ALIGNMENT) != 0)
+						{
+							UE_LOG(LogNNERuntimeIREE, Error, TEXT("CPU session: Input bindings memory need to be aligned with %d bytes"), IREE_HAL_HEAP_BUFFER_ALIGNMENT);
+							return ERunSyncStatus::Fail;
+						}
 
 						iree_hal_buffer_view_t* TempBufferView;
 						iree_hal_buffer_params_t Params = { 0 };
@@ -779,7 +816,7 @@ namespace UE::NNERuntimeIREE
 
 						if (!Device->CreateBufferViewCopy(InputTensorShapes[i].Rank(), Shape, IREEType, Params, (void*)InInputBindings[i].Data, InInputBindings[i].SizeInBytes, &TempBufferView))
 						{
-							UE::NNERuntimeIREE::Private::PrintIREEError("UE::NNERuntimeIREE::CPU::Private::FSession failed to allocate the buffer view", Status);
+							UE::NNERuntimeIREE::Private::PrintIREEError("CPU session: Failed to allocate the buffer view", Status);
 							if (TempBufferView)
 							{
 								iree_hal_buffer_view_release(TempBufferView);
@@ -792,20 +829,30 @@ namespace UE::NNERuntimeIREE
 						iree_hal_buffer_view_release(TempBufferView);
 						if (!iree_status_is_ok(Status))
 						{
-							UE::NNERuntimeIREE::Private::PrintIREEError("UE::NNERuntimeIREE::CPU::Private::FSession failed to push the buffer view to the input list", Status);
+							UE::NNERuntimeIREE::Private::PrintIREEError("CPU session: Failed to push the buffer view to the input list", Status);
 							iree_status_free(Status);
 							return ERunSyncStatus::Fail;
 						}
 					}
 
+					for (int32 i = 0; i < InOutputBindings.Num(); i++)
+					{
+						if (InOutputBindings[i].Data && FMath::Modulo<uint64>((uint64)InOutputBindings[i].Data, IREE_HAL_HEAP_BUFFER_ALIGNMENT) != 0)
+						{
+							UE_LOG(LogNNERuntimeIREE, Error, TEXT("CPU session: Output bindings memory need to be aligned with %d bytes"), IREE_HAL_HEAP_BUFFER_ALIGNMENT);
+							return ERunSyncStatus::Fail;
+						}
+					}
+
 					Status = iree_runtime_call_invoke(&Call, 0);
 					if (!iree_status_is_ok(Status))
 					{
-						UE::NNERuntimeIREE::Private::PrintIREEError("UE::NNERuntimeIREE::CPU::Private::FSession failed to call the model function", Status);
+						UE::NNERuntimeIREE::Private::PrintIREEError("CPU session: Failed to call the model function", Status);
 						iree_status_free(Status);
 						return ERunSyncStatus::Fail;
 					}
 
+					// Set output shapes
 					OutputTensorShapes.Reset();
 					TArray<iree_hal_buffer_view_t*> BufferViews;
 					iree_hal_buffer_view_t* BufferView = nullptr;
@@ -825,43 +872,35 @@ namespace UE::NNERuntimeIREE
 						Status = iree_runtime_call_outputs_pop_front_buffer_view(&Call, &BufferView);
 					}
 
-					bool bCopyResults = true;
-					if (InOutputBindings.Num() != OutputTensorShapes.Num())
-					{
-						bCopyResults = false;
-					}
-					for (int32 i = 0; i < InOutputBindings.Num() && bCopyResults; i++)
-					{
-						if (InOutputBindings[i].SizeInBytes < iree_hal_buffer_view_element_size(BufferViews[i]) * iree_hal_buffer_view_element_count(BufferViews[i]))
-						{
-							bCopyResults = false;
-						}
-					}
+					// Copy result back only if we have the required buffers available and big enought
 					ERunSyncStatus Result = ERunSyncStatus::Ok;
-					if (bCopyResults)
+					if (!InOutputBindings.IsEmpty())
 					{
+						check(OutputTensorShapes.Num() == InOutputBindings.Num())
 						for (int32 i = 0; i < InOutputBindings.Num(); i++)
 						{
-							iree_hal_buffer_t* Buffer = iree_hal_buffer_view_buffer(BufferViews[i]);
-							if (!Buffer)
-							{
-								UE_LOG(LogNNERuntimeIREE, Error, TEXT("UE::NNERuntimeIREE::CPU::Private::FSession failed to get the result buffer"));
-								Result = ERunSyncStatus::Fail;
-								break;
-							}
-
 							int32 DataSizeInBytes = iree_hal_buffer_view_element_size(BufferViews[i]) * iree_hal_buffer_view_element_count(BufferViews[i]);
-
-							iree_hal_buffer_mapping_t BufferMapping;
-							Status = iree_hal_buffer_map_range(Buffer, IREE_HAL_MAPPING_MODE_PERSISTENT, IREE_HAL_MEMORY_ACCESS_READ, 0, DataSizeInBytes, &BufferMapping);
-							if (!iree_status_is_ok(Status))
+							if (InOutputBindings[i].Data && InOutputBindings[i].SizeInBytes <= DataSizeInBytes)
 							{
-								UE::NNERuntimeIREE::Private::PrintIREEError("UE::NNERuntimeIREE::CPU::Private::FSession failed to map the result buffer", Status);
-								Result = ERunSyncStatus::Fail;
-								break;
+								iree_hal_buffer_t* Buffer = iree_hal_buffer_view_buffer(BufferViews[i]);
+								if (!Buffer)
+								{
+									UE_LOG(LogNNERuntimeIREE, Error, TEXT("CPU session: Failed failed to get the result buffer"));
+									Result = ERunSyncStatus::Fail;
+									break;
+								}
+
+								iree_hal_buffer_mapping_t BufferMapping;
+								Status = iree_hal_buffer_map_range(Buffer, IREE_HAL_MAPPING_MODE_PERSISTENT, IREE_HAL_MEMORY_ACCESS_READ, 0, DataSizeInBytes, &BufferMapping);
+								if (!iree_status_is_ok(Status))
+								{
+									UE::NNERuntimeIREE::Private::PrintIREEError("CPU session: Failed failed to map the result buffer", Status);
+									Result = ERunSyncStatus::Fail;
+									break;
+								}
+								FMemory::Memcpy(InOutputBindings[i].Data, BufferMapping.contents.data, DataSizeInBytes);
+								iree_hal_buffer_unmap_range(&BufferMapping);
 							}
-							FMemory::Memcpy(InOutputBindings[i].Data, BufferMapping.contents.data, DataSizeInBytes);
-							iree_hal_buffer_unmap_range(&BufferMapping);
 						}
 					}
 
