[NNE] Added NNERuntimeIREERDG prototype (wip).

- Also added vulkan-spirv target backend to IREE compiler (Windows)

#tests Windows and Linux, PIE and packaged Game tested IREE models, consoles packaged Game tested IREE models

#rnx

diff --git a/Engine/Plugins/Experimental/NNERuntimeIREE/NNERuntimeIREE.uplugin b/Engine/Plugins/Experimental/NNERuntimeIREE/NNERuntimeIREE.uplugin
--- a/Engine/Plugins/Experimental/NNERuntimeIREE/NNERuntimeIREE.uplugin
+++ b/Engine/Plugins/Experimental/NNERuntimeIREE/NNERuntimeIREE.uplugin
@@ -15,6 +15,11 @@
 	"IsExperimentalVersion": true,
 	"Installed": false,
 	"Modules": [
+		{
+			"Name": "IREEDriverRDG",
+			"Type": "Runtime",
+			"LoadingPhase": "Default"
+		},
 		{
 			"Name": "NNERuntimeIREE",
 			"Type": "Runtime",
@@ -24,6 +29,11 @@
 			"Name": "NNERuntimeIREEEditor",
 			"Type": "Editor",
 			"LoadingPhase": "Default"
+		},
+		{
+			"Name": "IREEDriverRDGShaders",
+			"Type": "Runtime",
+			"LoadingPhase": "PostConfigInit"
 		}
 	]
 }
\ No newline at end of file
diff --git a/Engine/Plugins/Experimental/NNERuntimeIREE/Shaders/NNERuntimeIREEAdd.usf b/Engine/Plugins/Experimental/NNERuntimeIREE/Shaders/NNERuntimeIREEAdd.usf
new file mode 100644
--- /dev/null
+++ b/Engine/Plugins/Experimental/NNERuntimeIREE/Shaders/NNERuntimeIREEAdd.usf
@@ -0,0 +1,18 @@
+// Copyright Epic Games, Inc. All Rights Reserved.
+
+#include "/Engine/Public/Platform.ush"
+#include "/Engine/Private/Common.ush"
+
+int Size;
+Buffer<float> A;
+Buffer<float> B;
+RWBuffer<float> Result;
+
+[numthreads(1, 1, 1)]
+void Add(in const uint3 DispatchThreadID : SV_DispatchThreadID)
+{
+	for (int i = 0; i < Size; i++)
+	{
+		Result[i] = A[i] + 2.0 * B[i];
+	}
+}
\ No newline at end of file
diff --git a/Engine/Plugins/Experimental/NNERuntimeIREE/Shaders/NNERuntimeIREEAdd_Cross.usf b/Engine/Plugins/Experimental/NNERuntimeIREE/Shaders/NNERuntimeIREEAdd_Cross.usf
new file mode 100644
--- /dev/null
+++ b/Engine/Plugins/Experimental/NNERuntimeIREE/Shaders/NNERuntimeIREEAdd_Cross.usf
@@ -0,0 +1,30 @@
+// Copyright Epic Games, Inc. All Rights Reserved.
+
+#include "/Engine/Public/Platform.ush"
+
+RWByteAddressBuffer _resource_var_0_0_ : register(u0);
+RWByteAddressBuffer _resource_var_0_1_ : register(u1);
+RWByteAddressBuffer _resource_var_0_2_ : register(u2);
+
+static uint3 gl_LocalInvocationID;
+struct SPIRV_Cross_Input
+{
+    uint3 gl_LocalInvocationID : SV_GroupThreadID;
+};
+
+void comp_main()
+{
+    for (uint _24 = gl_LocalInvocationID.x; int(_24) < int(4u); )
+    {
+        _resource_var_0_2_.Store(_24 * 4 + 0, asuint(asfloat(_resource_var_0_0_.Load(_24 * 4 + 0)) + 2.0* asfloat(_resource_var_0_1_.Load(_24 * 4 + 0))));
+        _24 += 64u;
+        continue;
+    }
+}
+
+[numthreads(64, 1, 1)]
+void main(SPIRV_Cross_Input stage_input)
+{
+    gl_LocalInvocationID = stage_input.gl_LocalInvocationID;
+    comp_main();
+}
\ No newline at end of file
diff --git a/Engine/Plugins/Experimental/NNERuntimeIREE/Source/IREEDriverRDG/IREEDriverRDG.Build.cs b/Engine/Plugins/Experimental/NNERuntimeIREE/Source/IREEDriverRDG/IREEDriverRDG.Build.cs
new file mode 100644
--- /dev/null
+++ b/Engine/Plugins/Experimental/NNERuntimeIREE/Source/IREEDriverRDG/IREEDriverRDG.Build.cs
@@ -0,0 +1,40 @@
+// Copyright Epic Games, Inc. All Rights Reserved.
+
+using UnrealBuildTool;
+
+public class IREEDriverRDG : ModuleRules
+{
+	public IREEDriverRDG(ReadOnlyTargetRules Target) : base(Target)
+	{
+		PCHUsage = PCHUsageMode.UseExplicitOrSharedPCHs;
+
+		PrivateDependencyModuleNames.AddRange
+		(
+			new string[]
+			{
+				"Core",
+				"CoreUObject",
+				"Engine",
+				"IREE",
+				"NNE",
+				"Projects",
+				"RenderCore",
+				"RHI",
+				"IREEDriverRDGShaders" // TOTO remove
+			}
+		);
+
+		if (Target.Platform == UnrealTargetPlatform.Win64)
+		{
+			PublicDefinitions.Add("WITH_IREE_DRIVER_RDG");
+		}
+		else if (Target.Platform == UnrealTargetPlatform.Linux)
+		{
+			PublicDefinitions.Add("WITH_IREE_DRIVER_RDG");
+		}
+		else if (Target.Platform == UnrealTargetPlatform.Mac)
+		{
+			PublicDefinitions.Add("WITH_IREE_DRIVER_RDG");
+		}
+	}
+}
diff --git a/Engine/Plugins/Experimental/NNERuntimeIREE/Source/IREEDriverRDG/Internal/IREEDriverRDG.h b/Engine/Plugins/Experimental/NNERuntimeIREE/Source/IREEDriverRDG/Internal/IREEDriverRDG.h
new file mode 100644
--- /dev/null
+++ b/Engine/Plugins/Experimental/NNERuntimeIREE/Source/IREEDriverRDG/Internal/IREEDriverRDG.h
@@ -0,0 +1,41 @@
+// Copyright Epic Games, Inc. All Rights Reserved.
+
+#pragma once
+
+#ifdef WITH_IREE_DRIVER_RDG
+
+#include "HAL/Platform.h"
+
+#if PLATFORM_MICROSOFT
+#include "Microsoft/AllowMicrosoftPlatformTypes.h"
+#include "Microsoft/AllowMicrosoftPlatformAtomics.h"
+#endif // PLATFORM_MICROSOFT
+THIRD_PARTY_INCLUDES_START
+#include "iree/hal/api.h"
+THIRD_PARTY_INCLUDES_END
+#if PLATFORM_MICROSOFT
+#include "Microsoft/HideMicrosoftPlatformAtomics.h"
+#include "Microsoft/HideMicrosoftPlatformTypes.h"
+#endif // PLATFORM_MICROSOFT
+
+#include "RenderGraphFwd.h"
+namespace UE::IREE::HAL::RDG
+{
+
+IREEDRIVERRDG_API iree_status_t DeviceCreate(iree_allocator_t HostAllocator, iree_hal_device_t** OutDevice);
+
+IREEDRIVERRDG_API iree_status_t BufferWrap(iree_allocator_t HostAllocator, iree_hal_allocator_t* DeviceAllocator, FRDGBuffer* InBuffer, iree_hal_buffer_t** OutBuffer);
+
+IREEDRIVERRDG_API FRDGBufferRef BufferHandle(iree_hal_buffer_t* Buffer);
+
+IREEDRIVERRDG_API FRDGBuilder* BufferBuilder(iree_hal_buffer_t* Buffer);
+
+
+
+IREEDRIVERRDG_API void DeviceAllocatorSetGraphBuilder(iree_hal_allocator_t* DeviceAllocator, FRDGBuilder& GraphBuilder);
+IREEDRIVERRDG_API FRDGBuilder& DeviceAllocatorGetGraphBuilder(iree_hal_allocator_t* DeviceAllocator);
+IREEDRIVERRDG_API void DeviceAllocatorResetGraphBuilder(iree_hal_allocator_t* DeviceAllocator);
+
+} // UE::IREE
+
+#endif // WITH_IREE_DRIVER_RDG
\ No newline at end of file
diff --git a/Engine/Plugins/Experimental/NNERuntimeIREE/Source/IREEDriverRDG/Internal/IREEDriverRDGLog.h b/Engine/Plugins/Experimental/NNERuntimeIREE/Source/IREEDriverRDG/Internal/IREEDriverRDGLog.h
new file mode 100644
--- /dev/null
+++ b/Engine/Plugins/Experimental/NNERuntimeIREE/Source/IREEDriverRDG/Internal/IREEDriverRDGLog.h
@@ -0,0 +1,7 @@
+// Copyright Epic Games, Inc. All Rights Reserved.
+
+#pragma once
+
+#include "Logging/LogMacros.h"
+
+DECLARE_LOG_CATEGORY_EXTERN(LogIREEDriverRDG, Log, All);
\ No newline at end of file
diff --git a/Engine/Plugins/Experimental/NNERuntimeIREE/Source/IREEDriverRDG/Private/IREEDriverRDGBuffer.cpp b/Engine/Plugins/Experimental/NNERuntimeIREE/Source/IREEDriverRDG/Private/IREEDriverRDGBuffer.cpp
new file mode 100644
--- /dev/null
+++ b/Engine/Plugins/Experimental/NNERuntimeIREE/Source/IREEDriverRDG/Private/IREEDriverRDGBuffer.cpp
@@ -0,0 +1,174 @@
+// Copyright Epic Games, Inc. All Rights Reserved.
+
+#include "IREEDriverRDGBuffer.h"
+
+#ifdef WITH_IREE_DRIVER_RDG
+
+#include "RenderGraphUtils.h"
+#include "IREEDriverRDGLog.h"
+
+namespace UE::IREE::HAL::RDG
+{
+
+namespace Private
+{
+
+iree_hal_memory_type_bits_t GetIREEMemoryTypeFromBufferUsage(EBufferUsageFlags Usage)
+{
+	//TODO: implement propperly
+	return IREE_HAL_MEMORY_TYPE_DEVICE_LOCAL;
+}
+
+iree_hal_memory_access_bits_t GetIREEMemoryAccessFromBufferUsage(EBufferUsageFlags Usage)
+{
+	//TODO: implement propperly
+	return IREE_HAL_MEMORY_ACCESS_ALL;
+}
+
+iree_hal_buffer_usage_bits_t GetIREEBufferUsageFromBufferUsage(EBufferUsageFlags Usage)
+{
+	//TODO: implement propperly
+	return static_cast<iree_hal_buffer_usage_bits_t>(IREE_HAL_BUFFER_USAGE_TRANSFER | IREE_HAL_BUFFER_USAGE_DISPATCH_STORAGE);
+}
+
+class FBuffer
+{
+public:
+	static iree_status_t Create(iree_allocator_t HostAllocator, iree_hal_allocator_t* DeviceAllocator, const iree_hal_buffer_params_t* Params, iree_device_size_t AllocationSize, iree_hal_buffer_t** OutBuffer)
+	{
+		check(DeviceAllocator);
+		check(Params);
+		check(OutBuffer);
+		UE_LOG(LogIREEDriverRDG, Display, TEXT("%s"), StringCast<TCHAR>(__FUNCTION__).Get());
+
+		FBuffer* Buffer;
+		IREE_RETURN_IF_ERROR(iree_allocator_malloc(HostAllocator, sizeof(*Buffer), (void**)&Buffer));
+		iree_hal_buffer_initialize(HostAllocator, DeviceAllocator, &Buffer->Base, AllocationSize,
+									0, AllocationSize, Params->type, Params->access, Params->usage,
+									&FBuffer::VTable, &Buffer->Base);
+
+		FRDGBuilder& GraphBuilder = DeviceAllocatorGetGraphBuilder(DeviceAllocator); // TODO
+
+		FRDGBufferDesc BufferDesc = FRDGBufferDesc::CreateBufferDesc(1, AllocationSize); // TODO type, flags, usage, access, alignment etc.
+		Buffer->Buffer = GraphBuilder.CreateBuffer(BufferDesc, TEXT("IREE::HAL::RDG::BufferCreate"));
+
+		*OutBuffer = (iree_hal_buffer_t*)Buffer;
+		return iree_ok_status();
+	}
+
+	static iree_status_t BufferWrap(iree_allocator_t HostAllocator, iree_hal_allocator_t* DeviceAllocator, FRDGBufferRef RDGBuffer, iree_hal_buffer_t** OutBuffer)
+	{
+		check(DeviceAllocator);
+		check(RDGBuffer);
+		check(OutBuffer);
+		UE_LOG(LogIREEDriverRDG, Display, TEXT("%s"), StringCast<TCHAR>(__FUNCTION__).Get());
+
+		FBuffer* Buffer;
+		IREE_RETURN_IF_ERROR(iree_allocator_malloc(HostAllocator, sizeof(*Buffer), (void**)&Buffer));
+		
+		FRDGBufferDesc Desc = RDGBuffer->Desc;
+		iree_hal_buffer_initialize(HostAllocator, DeviceAllocator, &Buffer->Base, Desc.GetSize(), 0, Desc.GetSize(), GetIREEMemoryTypeFromBufferUsage(Desc.Usage), GetIREEMemoryAccessFromBufferUsage(Desc.Usage), GetIREEBufferUsageFromBufferUsage(Desc.Usage), &FBuffer::VTable, &Buffer->Base);
+		Buffer->Buffer = RDGBuffer;
+
+		*OutBuffer = (iree_hal_buffer_t*)Buffer;
+		
+		return iree_ok_status();
+	}
+
+	FRDGBufferRef GetRDGBufferHandle()
+	{
+		return Buffer;
+	}
+
+private:
+	static FBuffer* Cast(const iree_hal_buffer_t* Buffer)
+	{
+		checkf(iree_hal_resource_is(Buffer, &FBuffer::VTable), TEXT("FBuffer: type does not match"));
+		return (FBuffer*)Buffer;
+	}
+
+	static void Recycle(iree_hal_buffer_t *BaseBuffer)
+	{
+		UE_LOG(LogIREEDriverRDG, Display, TEXT("%s"), StringCast<TCHAR>(__FUNCTION__).Get());
+		return;
+	}
+
+	static void Destroy(iree_hal_buffer_t *BaseBuffer)
+	{
+		UE_LOG(LogIREEDriverRDG, Display, TEXT("%s"), StringCast<TCHAR>(__FUNCTION__).Get());
+		FBuffer* Buffer = Cast(BaseBuffer);
+		iree_allocator_free(Buffer->Base.host_allocator, Buffer);
+	}
+
+	static iree_status_t MapRange(iree_hal_buffer_t *BaseBuffer,
+									iree_hal_mapping_mode_t MappingMode,
+									iree_hal_memory_access_t MemoryAccess,
+									iree_device_size_t LocalByteOffset,
+									iree_device_size_t LocalByteLength,
+									iree_hal_buffer_mapping_t* Mapping)
+	{
+		UE_LOG(LogIREEDriverRDG, Display, TEXT("%s"), StringCast<TCHAR>(__FUNCTION__).Get());
+		return iree_make_status(IREE_STATUS_UNIMPLEMENTED, __FUNCTION__);
+	}
+
+	static iree_status_t UnmapRange(iree_hal_buffer_t *BaseBuffer,
+									iree_device_size_t LocalByteOffset,
+									iree_device_size_t LocalByteLength,
+									iree_hal_buffer_mapping_t* Mapping)
+	{
+		UE_LOG(LogIREEDriverRDG, Display, TEXT("%s"), StringCast<TCHAR>(__FUNCTION__).Get());
+		return iree_make_status(IREE_STATUS_UNIMPLEMENTED, __FUNCTION__);
+	}
+
+	static iree_status_t InvalidateRange(iree_hal_buffer_t *BaseBuffer,
+											iree_device_size_t LocalByteOffset,
+											iree_device_size_t LocalByteLength)
+	{
+		UE_LOG(LogIREEDriverRDG, Display, TEXT("%s"), StringCast<TCHAR>(__FUNCTION__).Get());
+		return iree_make_status(IREE_STATUS_UNIMPLEMENTED, __FUNCTION__);
+	}
+
+	static iree_status_t FlushRange(iree_hal_buffer_t *BaseBuffer,
+									iree_device_size_t LocalByteOffset,
+									iree_device_size_t LocalByteLength)
+	{
+		UE_LOG(LogIREEDriverRDG, Display, TEXT("%s"), StringCast<TCHAR>(__FUNCTION__).Get());
+		return iree_make_status(IREE_STATUS_UNIMPLEMENTED, __FUNCTION__);
+	}
+
+	static const iree_hal_buffer_vtable_t VTable;
+
+	iree_hal_buffer_t Base;
+	FRDGBufferRef Buffer;
+};
+
+const iree_hal_buffer_vtable_t FBuffer::VTable =
+{
+	.recycle = FBuffer::Recycle,
+	.destroy = FBuffer::Destroy,
+	.map_range = FBuffer::MapRange,
+	.unmap_range = FBuffer::UnmapRange,
+	.invalidate_range = FBuffer::InvalidateRange,
+	.flush_range = FBuffer::FlushRange
+};
+
+} // namespace Private
+
+iree_status_t BufferWrap(iree_allocator_t HostAllocator, iree_hal_allocator_t* DeviceAllocator, FRDGBufferRef RDGBuffer, iree_hal_buffer_t** OutBuffer)
+{
+	return Private::FBuffer::BufferWrap(HostAllocator, DeviceAllocator, RDGBuffer, OutBuffer);
+}
+
+iree_status_t BufferCreate(iree_allocator_t HostAllocator, iree_hal_allocator_t* DeviceAllocator, const iree_hal_buffer_params_t* Params, iree_device_size_t AllocationSize, iree_hal_buffer_t** OutBuffer)
+{
+	return Private::FBuffer::Create(HostAllocator, DeviceAllocator, Params, AllocationSize, OutBuffer);
+}
+
+FRDGBufferRef BufferHandle(iree_hal_buffer_t* Buffer)
+{
+	return Buffer ? ((Private::FBuffer*)Buffer)->GetRDGBufferHandle() : nullptr;
+}
+
+} // UE::IREE
+
+#endif // WITH_IREE_DRIVER_RDG
\ No newline at end of file
diff --git a/Engine/Plugins/Experimental/NNERuntimeIREE/Source/IREEDriverRDG/Private/IREEDriverRDGBuffer.h b/Engine/Plugins/Experimental/NNERuntimeIREE/Source/IREEDriverRDG/Private/IREEDriverRDGBuffer.h
new file mode 100644
--- /dev/null
+++ b/Engine/Plugins/Experimental/NNERuntimeIREE/Source/IREEDriverRDG/Private/IREEDriverRDGBuffer.h
@@ -0,0 +1,16 @@
+// Copyright Epic Games, Inc. All Rights Reserved.
+
+#pragma once
+
+#ifdef WITH_IREE_DRIVER_RDG
+
+#include "IREEDriverRDG.h"
+
+namespace UE::IREE::HAL::RDG
+{
+
+iree_status_t BufferCreate(iree_allocator_t HostAllocator, iree_hal_allocator_t* DeviceAllocator, const iree_hal_buffer_params_t* Params, iree_device_size_t AllocationSize, iree_hal_buffer_t** OutBuffer);
+
+} // UE::IREE
+
+#endif // WITH_IREE_DRIVER_RDG
\ No newline at end of file
diff --git a/Engine/Plugins/Experimental/NNERuntimeIREE/Source/IREEDriverRDG/Private/IREEDriverRDGCommandBuffer.cpp b/Engine/Plugins/Experimental/NNERuntimeIREE/Source/IREEDriverRDG/Private/IREEDriverRDGCommandBuffer.cpp
new file mode 100644
--- /dev/null
+++ b/Engine/Plugins/Experimental/NNERuntimeIREE/Source/IREEDriverRDG/Private/IREEDriverRDGCommandBuffer.cpp
@@ -0,0 +1,333 @@
+// Copyright Epic Games, Inc. All Rights Reserved.
+
+#include "IREEDriverRDGCommandBuffer.h"
+
+#ifdef WITH_IREE_DRIVER_RDG
+
+#include "IREEDriverRDGBuffer.h"
+#include "IREEDriverRDGLog.h"
+#include "RenderGraphBuilder.h"
+#include "RenderGraphUtils.h"
+
+
+#include "IREEDriverRDGShadersAddCrossCS.h"
+
+DECLARE_GPU_STAT_NAMED(FCommandBufferDispatch, TEXT("CommandBuffer.Dispatch"));
+namespace UE::IREE::HAL::RDG
+{
+
+namespace Private
+{
+
+struct FBinding
+{
+	FRDGBufferRef Buffer = nullptr;
+};
+
+class FCommandBuffer
+{
+public:
+	static iree_status_t Create(iree_allocator_t HostAllocator,
+								iree_hal_device_t* Device,
+								iree_hal_command_buffer_mode_t Mode,
+								iree_hal_command_category_t CommandCategories,	
+								iree_hal_queue_affinity_t QueueAffinity,
+								iree_host_size_t BindingCapacity,
+								iree_hal_command_buffer_t** OutCommandBuffer)
+	{
+		check(OutCommandBuffer);
+		UE_LOG(LogIREEDriverRDG, Display, TEXT("%s"), StringCast<TCHAR>(__FUNCTION__).Get());
+
+		FCommandBuffer* CommandBuffer;
+		IREE_RETURN_IF_ERROR(iree_allocator_malloc(HostAllocator, sizeof(*CommandBuffer), (void**)&CommandBuffer));
+		iree_hal_command_buffer_initialize(Device, Mode, CommandCategories,
+											QueueAffinity, BindingCapacity,
+											&FCommandBuffer::VTable, &CommandBuffer->Base);
+		CommandBuffer->HostAllocator = HostAllocator;
+		CommandBuffer->Device = Device;
+
+		*OutCommandBuffer = (iree_hal_command_buffer_t*)CommandBuffer;
+		return iree_ok_status();
+	}
+
+	iree_status_t BindDescriptorSet(iree_hal_pipeline_layout_t* PipelineLayout, uint32 Set, iree_host_size_t BindingCount, const iree_hal_descriptor_set_binding_t* Bindings)
+	{
+		for (int i = 0; i < BindingCount; i++) {
+			const iree_hal_descriptor_set_binding_t& IREEBinding = Bindings[i];
+
+			check(IREEBinding.binding == i);
+			check(IREEBinding.buffer);
+			check(IREEBinding.offset == 0);
+
+			FBinding Binding;
+			Binding.Buffer = BufferHandle(IREEBinding.buffer);
+
+			RecordedBindings.Add(Binding);
+		}
+		return iree_ok_status();
+	}
+
+	iree_status_t Dispatch(iree_hal_executable_t* Executable, int32 EntryPoint, uint32 WorkgroupX, uint32 WorkgroupY, uint32 WorkgroupZ)
+	{
+		using namespace UE::IREEDriverRDGShaders::Internal;
+
+		iree_hal_allocator_t* DeviceAllocator = iree_hal_device_allocator(Device);
+		FRDGBuilder& GraphBuilder = DeviceAllocatorGetGraphBuilder(DeviceAllocator); // TODO 
+
+		check(RecordedBindings.Num() == 3);
+
+		{
+			FRDGBufferDesc ACopyDesc = FRDGBufferDesc::CreateByteAddressDesc(RecordedBindings[0].Buffer->Desc.GetSize());
+			FRDGBufferDesc BCopyDesc = FRDGBufferDesc::CreateByteAddressDesc(RecordedBindings[1].Buffer->Desc.GetSize());
+			FRDGBufferDesc ResultCopyDesc = FRDGBufferDesc::CreateByteAddressDesc(RecordedBindings[2].Buffer->Desc.GetSize());
+
+			FRDGBufferRef ACopy = GraphBuilder.CreateBuffer(ACopyDesc, TEXT("CopyA"));
+			FRDGBufferRef BCopy = GraphBuilder.CreateBuffer(BCopyDesc, TEXT("CopyB"));
+			FRDGBufferRef ResultCopy = GraphBuilder.CreateBuffer(ResultCopyDesc, TEXT("CopyResult"));
+
+			AddCopyBufferPass(GraphBuilder, ACopy, RecordedBindings[0].Buffer);
+			AddCopyBufferPass(GraphBuilder, BCopy, RecordedBindings[1].Buffer);
+
+			FAddCrossCS::FParameters *ShaderParameters = GraphBuilder.AllocParameters<FAddCrossCS::FParameters>();
+			ShaderParameters->_resource_var_0_0_ = GraphBuilder.CreateUAV(ACopy);
+			ShaderParameters->_resource_var_0_1_ = GraphBuilder.CreateUAV(BCopy);
+			ShaderParameters->_resource_var_0_2_ = GraphBuilder.CreateUAV(ResultCopy);
+
+			FGlobalShaderMap* GlobalShaderMap = GetGlobalShaderMap(GMaxRHIFeatureLevel);
+			TShaderMapRef<FAddCrossCS> Shader(GlobalShaderMap);
+
+			RDG_EVENT_SCOPE_STAT(GraphBuilder, FCommandBufferDispatch, "CommandBuffer.Dispatch");
+			RDG_GPU_STAT_SCOPE(GraphBuilder, FCommandBufferDispatch);
+
+			FComputeShaderUtils::AddPass(
+				GraphBuilder,
+				RDG_EVENT_NAME("CommandBuffer.Dispatch"),
+				ERDGPassFlags::Compute | ERDGPassFlags::NeverCull,
+				Shader,
+				ShaderParameters,
+				FIntVector(WorkgroupX, WorkgroupY, WorkgroupZ));
+
+			AddCopyBufferPass(GraphBuilder, RecordedBindings[2].Buffer, ResultCopy);
+		}
+
+		return iree_ok_status();
+	}
+
+private:
+	static FCommandBuffer* Cast(iree_hal_command_buffer_t* CommandBuffer)
+	{
+		checkf(iree_hal_resource_is(CommandBuffer, &FCommandBuffer::VTable), TEXT("FCommandBuffer: type does not match"));
+		return (FCommandBuffer*)CommandBuffer;
+	}
+
+	static void Destroy(iree_hal_command_buffer_t* BaseCommandBuffer)
+	{
+		UE_LOG(LogIREEDriverRDG, Display, TEXT("%s"), StringCast<TCHAR>(__FUNCTION__).Get());
+		FCommandBuffer* CommandBuffer = Cast(BaseCommandBuffer);
+		iree_allocator_free(CommandBuffer->HostAllocator, CommandBuffer);
+	}
+
+	static iree_status_t Begin(iree_hal_command_buffer_t* BaseCommandBuffer)
+	{
+		UE_LOG(LogIREEDriverRDG, Display, TEXT("%s"), StringCast<TCHAR>(__FUNCTION__).Get());
+		//return iree_make_status(IREE_STATUS_UNIMPLEMENTED, __FUNCTION__);
+		return iree_ok_status();
+	}
+
+	static iree_status_t End(iree_hal_command_buffer_t* BaseCommandBuffer)
+	{
+		UE_LOG(LogIREEDriverRDG, Display, TEXT("%s"), StringCast<TCHAR>(__FUNCTION__).Get());
+		//return iree_make_status(IREE_STATUS_UNIMPLEMENTED, __FUNCTION__);
+		return iree_ok_status();
+	}
+
+	static void BeginDebugGroup(iree_hal_command_buffer_t* BaseCommandBuffer,
+											iree_string_view_t Label,
+											iree_hal_label_color_t LabelColor,
+											const iree_hal_label_location_t* Location)
+	{
+		UE_LOG(LogIREEDriverRDG, Display, TEXT("%s"), StringCast<TCHAR>(__FUNCTION__).Get());
+	}
+
+	static void EndDebugGroup(iree_hal_command_buffer_t* BaseCommandBuffer)
+	{
+		UE_LOG(LogIREEDriverRDG, Display, TEXT("%s"), StringCast<TCHAR>(__FUNCTION__).Get());
+	}
+
+	static iree_status_t ExecutionBarrier(iree_hal_command_buffer_t* BaseCommandBuffer,
+											iree_hal_execution_stage_t SourceStageMask,
+											iree_hal_execution_stage_t TargetStageMask,
+											iree_hal_execution_barrier_flags_t Flags,
+											iree_host_size_t MemoryBarrierCount,
+											const iree_hal_memory_barrier_t* MemoryBarriers,
+											iree_host_size_t BufferBarrierCount,
+											const iree_hal_buffer_barrier_t* BufferBarriers)
+	{
+		UE_LOG(LogIREEDriverRDG, Display, TEXT("%s"), StringCast<TCHAR>(__FUNCTION__).Get());
+		// return iree_make_status(IREE_STATUS_UNIMPLEMENTED, __FUNCTION__);
+		return iree_ok_status();
+	}
+
+	static iree_status_t SignalEvent(iree_hal_command_buffer_t* BaseCommandBuffer,
+										iree_hal_event_t* Event,
+										iree_hal_execution_stage_t SourceStageMask)
+	{
+		UE_LOG(LogIREEDriverRDG, Display, TEXT("%s"), StringCast<TCHAR>(__FUNCTION__).Get());
+		return iree_make_status(IREE_STATUS_UNIMPLEMENTED, __FUNCTION__);
+	}
+
+	static iree_status_t ResetEvent(iree_hal_command_buffer_t* BaseCommandBuffer,
+									iree_hal_event_t* Event,
+									iree_hal_execution_stage_t SourceStageMask)
+	{
+		UE_LOG(LogIREEDriverRDG, Display, TEXT("%s"), StringCast<TCHAR>(__FUNCTION__).Get());
+		return iree_make_status(IREE_STATUS_UNIMPLEMENTED, __FUNCTION__);
+	}
+
+	static iree_status_t WaitEvents(iree_hal_command_buffer_t* BaseCommandBuffer,
+									iree_host_size_t EventCount,
+									const iree_hal_event_t** Events,
+									iree_hal_execution_stage_t SourceStageMask,
+									iree_hal_execution_stage_t TargetStageMask,
+									iree_host_size_t MemoryBarrierCount,
+									const iree_hal_memory_barrier_t* MemoryBarriers,
+									iree_host_size_t BufferBarrierCount,
+									const iree_hal_buffer_barrier_t* BufferBarriers)
+	{
+		UE_LOG(LogIREEDriverRDG, Display, TEXT("%s"), StringCast<TCHAR>(__FUNCTION__).Get());
+		return iree_make_status(IREE_STATUS_UNIMPLEMENTED, __FUNCTION__);
+	}
+
+	static iree_status_t DiscardBuffer(iree_hal_command_buffer_t* BaseCommandBuffer,
+										iree_hal_buffer_t* Buffer)
+	{
+		UE_LOG(LogIREEDriverRDG, Display, TEXT("%s"), StringCast<TCHAR>(__FUNCTION__).Get());
+		return iree_make_status(IREE_STATUS_UNIMPLEMENTED, __FUNCTION__);
+	}
+
+	static iree_status_t FillBuffer(iree_hal_command_buffer_t* BaseCommandBuffer,
+									iree_hal_buffer_t* TargetBuffer,
+									iree_device_size_t TargetOffset,
+									iree_device_size_t Length,
+									const void* Pattern,
+									iree_host_size_t PatternLength)
+	{
+		UE_LOG(LogIREEDriverRDG, Display, TEXT("%s"), StringCast<TCHAR>(__FUNCTION__).Get());
+		return iree_make_status(IREE_STATUS_UNIMPLEMENTED, __FUNCTION__);
+	}
+
+	static iree_status_t UpdateBuffer(iree_hal_command_buffer_t* BaseCommandBuffer,
+										const void* SourceBuffer,
+										iree_host_size_t SourceOffset,
+										iree_hal_buffer_t* TargetBuffer,
+										iree_device_size_t TargetOffset,
+										iree_device_size_t Length)
+	{
+		UE_LOG(LogIREEDriverRDG, Display, TEXT("%s"), StringCast<TCHAR>(__FUNCTION__).Get());
+		//return iree_make_status(IREE_STATUS_UNIMPLEMENTED, __FUNCTION__);
+		return iree_ok_status();
+	}
+
+	static iree_status_t CopyBuffer(iree_hal_command_buffer_t* BaseCommandBuffer,
+									iree_hal_buffer_t* SourceBuffer,
+									iree_device_size_t SourceOffset,
+									iree_hal_buffer_t* TargetBuffer,
+									iree_device_size_t TargetOffset,
+									iree_device_size_t Length)
+	{
+		UE_LOG(LogIREEDriverRDG, Display, TEXT("%s"), StringCast<TCHAR>(__FUNCTION__).Get());
+		return iree_make_status(IREE_STATUS_UNIMPLEMENTED, __FUNCTION__);
+	}
+
+	static iree_status_t Collective(iree_hal_command_buffer_t* BaseCommandBuffer, iree_hal_channel_t* Channel, iree_hal_collective_op_t Op, uint32 Param, iree_hal_buffer_binding_t SendBinding, iree_hal_buffer_binding_t ReceiveBinding, iree_device_size_t ElementCount)
+	{
+		UE_LOG(LogIREEDriverRDG, Display, TEXT("%s"), StringCast<TCHAR>(__FUNCTION__).Get());
+		return iree_make_status(IREE_STATUS_UNIMPLEMENTED, __FUNCTION__);
+	}
+
+	static iree_status_t PushConstants(iree_hal_command_buffer_t* BaseCommandBuffer, iree_hal_pipeline_layout_t* PipelineLayout, iree_host_size_t Ofsset, const void* Values, iree_host_size_t ValuesLength)
+	{
+		UE_LOG(LogIREEDriverRDG, Display, TEXT("%s"), StringCast<TCHAR>(__FUNCTION__).Get());
+		return iree_make_status(IREE_STATUS_UNIMPLEMENTED, __FUNCTION__);
+	}
+
+	static iree_status_t PushDescriptorSet(iree_hal_command_buffer_t* BaseCommandBuffer, iree_hal_pipeline_layout_t* PipelineLayout, uint32 Set, iree_host_size_t BindingCount, const iree_hal_descriptor_set_binding_t* Bindings)
+	{
+		UE_LOG(LogIREEDriverRDG, Display, TEXT("%s"), StringCast<TCHAR>(__FUNCTION__).Get());
+		
+		return Cast(BaseCommandBuffer)->BindDescriptorSet(PipelineLayout, Set, BindingCount, Bindings);
+	}
+
+	static iree_status_t Dispatch(iree_hal_command_buffer_t* BaseCommandBuffer, iree_hal_executable_t* Executable, int32 EntryPoint, uint32 WorkgroupX, uint32 WorkgroupY, uint32 WorkgroupZ)
+	{
+		UE_LOG(LogIREEDriverRDG, Display, TEXT("%s"), StringCast<TCHAR>(__FUNCTION__).Get());
+		
+		return Cast(BaseCommandBuffer)->Dispatch(Executable, EntryPoint, WorkgroupX, WorkgroupY, WorkgroupZ);
+	}
+
+	static iree_status_t DispatchIndirect(iree_hal_command_buffer_t* BaseCommandBuffer,
+											iree_hal_executable_t* Executable,
+											int32 EntryPoint,
+											iree_hal_buffer_t* WorkgroupsBuffer,
+											iree_device_size_t WorkgroupsOffset)
+	{
+		UE_LOG(LogIREEDriverRDG, Display, TEXT("%s"), StringCast<TCHAR>(__FUNCTION__).Get());
+		return iree_make_status(IREE_STATUS_UNIMPLEMENTED, __FUNCTION__);
+	}
+	static iree_status_t ExecuteCommands(iree_hal_command_buffer_t* BaseCommandBuffer,
+											iree_hal_command_buffer_t* Commands,
+											iree_hal_buffer_binding_table_t BindingTable)
+	{
+		UE_LOG(LogIREEDriverRDG, Display, TEXT("%s"), StringCast<TCHAR>(__FUNCTION__).Get());
+		return iree_make_status(IREE_STATUS_UNIMPLEMENTED, __FUNCTION__);
+	}
+	
+	 static const iree_hal_command_buffer_vtable_t VTable;
+
+	iree_hal_command_buffer_t Base;
+	iree_allocator_t HostAllocator;
+	iree_hal_device_t* Device;
+
+	TArray<FBinding> RecordedBindings;
+};
+
+const iree_hal_command_buffer_vtable_t FCommandBuffer::VTable = 
+{
+	.destroy = FCommandBuffer::Destroy,
+	.begin = FCommandBuffer::Begin,
+	.end = FCommandBuffer::End,
+	.begin_debug_group = FCommandBuffer::BeginDebugGroup,
+	.end_debug_group = FCommandBuffer::EndDebugGroup,
+	.execution_barrier = FCommandBuffer::ExecutionBarrier,
+	.signal_event = FCommandBuffer::SignalEvent,
+	.reset_event = FCommandBuffer::ResetEvent,
+	.wait_events = FCommandBuffer::WaitEvents,
+	.discard_buffer = FCommandBuffer::DiscardBuffer,
+	.fill_buffer = FCommandBuffer::FillBuffer,
+	.update_buffer = FCommandBuffer::UpdateBuffer,
+	.copy_buffer = FCommandBuffer::CopyBuffer,
+	.collective = FCommandBuffer::Collective,
+	.push_constants = FCommandBuffer::PushConstants,
+	.push_descriptor_set = FCommandBuffer::PushDescriptorSet,
+	.dispatch = FCommandBuffer::Dispatch,
+	.dispatch_indirect = FCommandBuffer::DispatchIndirect,
+	.execute_commands = FCommandBuffer::ExecuteCommands
+};
+
+} // namespace Private
+
+iree_status_t CommandBufferCreate(
+	iree_allocator_t HostAllocator,
+	iree_hal_device_t* Device,
+	iree_hal_command_buffer_mode_t Mode,
+	iree_hal_command_category_t CommandCategories,	
+	iree_hal_queue_affinity_t QueueAffinity,
+	iree_host_size_t BindingCapacity,
+	iree_hal_command_buffer_t** OutCommandBuffer)
+{
+	return Private::FCommandBuffer::Create(HostAllocator, Device, Mode, CommandCategories, QueueAffinity, BindingCapacity, OutCommandBuffer);
+}
+
+} // UE::IREE
+
+#endif // WITH_IREE_DRIVER_RDG
\ No newline at end of file
diff --git a/Engine/Plugins/Experimental/NNERuntimeIREE/Source/IREEDriverRDG/Private/IREEDriverRDGCommandBuffer.h b/Engine/Plugins/Experimental/NNERuntimeIREE/Source/IREEDriverRDG/Private/IREEDriverRDGCommandBuffer.h
new file mode 100644
--- /dev/null
+++ b/Engine/Plugins/Experimental/NNERuntimeIREE/Source/IREEDriverRDG/Private/IREEDriverRDGCommandBuffer.h
@@ -0,0 +1,23 @@
+// Copyright Epic Games, Inc. All Rights Reserved.
+
+#pragma once
+
+#ifdef WITH_IREE_DRIVER_RDG
+
+#include "IREEDriverRDG.h"
+
+namespace UE::IREE::HAL::RDG
+{
+
+iree_status_t CommandBufferCreate(
+	iree_allocator_t HostAllocator,
+	iree_hal_device_t* Device,
+	iree_hal_command_buffer_mode_t Mode,
+	iree_hal_command_category_t CommandCategories,	
+	iree_hal_queue_affinity_t QueueAffinity,
+	iree_host_size_t BindingCapacity,
+	iree_hal_command_buffer_t** OutCommandBuffer);
+
+} // UE::IREE
+
+#endif // WITH_IREE_DRIVER_RDG
\ No newline at end of file
diff --git a/Engine/Plugins/Experimental/NNERuntimeIREE/Source/IREEDriverRDG/Private/IREEDriverRDGDescriptorSetLayout.cpp b/Engine/Plugins/Experimental/NNERuntimeIREE/Source/IREEDriverRDG/Private/IREEDriverRDGDescriptorSetLayout.cpp
new file mode 100644
--- /dev/null
+++ b/Engine/Plugins/Experimental/NNERuntimeIREE/Source/IREEDriverRDG/Private/IREEDriverRDGDescriptorSetLayout.cpp
@@ -0,0 +1,65 @@
+// Copyright Epic Games, Inc. All Rights Reserved.
+
+#include "IREEDriverRDGDescriptorSetLayout.h"
+
+#ifdef WITH_IREE_DRIVER_RDG
+
+#include "IREEDriverRDGLog.h"
+
+namespace UE::IREE::HAL::RDG
+{
+
+namespace Private
+{
+
+class FDescriptorSetLayout
+{
+private:
+	iree_hal_resource_t Resource;
+	iree_allocator_t HostAllocator;
+public:
+	static iree_status_t Create(iree_allocator_t HostAllocator, iree_hal_descriptor_set_layout_t** OutDescriptorSetLayout)
+	{
+		check(OutDescriptorSetLayout);
+		UE_LOG(LogIREEDriverRDG, Display, TEXT("%s"), StringCast<TCHAR>(__FUNCTION__).Get());
+
+		FDescriptorSetLayout* DescriptorSetLayout;
+		IREE_RETURN_IF_ERROR(iree_allocator_malloc(HostAllocator, sizeof(*DescriptorSetLayout), (void**)&DescriptorSetLayout));
+		iree_hal_resource_initialize((const void*)&FDescriptorSetLayout::VTable, &DescriptorSetLayout->Resource);
+		DescriptorSetLayout->HostAllocator = HostAllocator;
+
+		*OutDescriptorSetLayout = (iree_hal_descriptor_set_layout_t*)DescriptorSetLayout;
+		return iree_ok_status();
+	}
+private:
+	static FDescriptorSetLayout* Cast(iree_hal_descriptor_set_layout_t* DescriptorSetLayout)
+	{
+		checkf(iree_hal_resource_is(DescriptorSetLayout, &FDescriptorSetLayout::VTable), TEXT("FDescriptorSetLayout: type does not match"));
+		return (FDescriptorSetLayout*)DescriptorSetLayout;
+	}
+
+	static void Destroy(iree_hal_descriptor_set_layout_t* BaseDescriptorSetLayout)
+	{
+		UE_LOG(LogIREEDriverRDG, Display, TEXT("%s"), StringCast<TCHAR>(__FUNCTION__).Get());
+		FDescriptorSetLayout* DescriptorSetLayout = Cast(BaseDescriptorSetLayout);
+		iree_allocator_free(DescriptorSetLayout->HostAllocator, DescriptorSetLayout);
+	}
+
+	static const iree_hal_descriptor_set_layout_vtable_t VTable;
+};
+
+const iree_hal_descriptor_set_layout_vtable_t FDescriptorSetLayout::VTable =
+{
+	.destroy = FDescriptorSetLayout::Destroy,
+};
+
+} // namespace Private
+
+iree_status_t DescriptorSetLayoutCreate(iree_allocator_t HostAllocator, iree_hal_descriptor_set_layout_t** OutDescriptorSetLayout)
+{
+	return Private::FDescriptorSetLayout::Create(HostAllocator, OutDescriptorSetLayout);
+}
+
+} // UE::IREE
+
+#endif // WITH_IREE_DRIVER_RDG
\ No newline at end of file
diff --git a/Engine/Plugins/Experimental/NNERuntimeIREE/Source/IREEDriverRDG/Private/IREEDriverRDGDescriptorSetLayout.h b/Engine/Plugins/Experimental/NNERuntimeIREE/Source/IREEDriverRDG/Private/IREEDriverRDGDescriptorSetLayout.h
new file mode 100644
--- /dev/null
+++ b/Engine/Plugins/Experimental/NNERuntimeIREE/Source/IREEDriverRDG/Private/IREEDriverRDGDescriptorSetLayout.h
@@ -0,0 +1,16 @@
+// Copyright Epic Games, Inc. All Rights Reserved.
+
+#pragma once
+
+#ifdef WITH_IREE_DRIVER_RDG
+
+#include "IREEDriverRDG.h"
+
+namespace UE::IREE::HAL::RDG
+{
+
+iree_status_t DescriptorSetLayoutCreate(iree_allocator_t HostAllocator, iree_hal_descriptor_set_layout_t** OutDescriptorSetLayout);
+
+} // UE::IREE
+
+#endif // WITH_IREE_DRIVER_RDG
\ No newline at end of file
diff --git a/Engine/Plugins/Experimental/NNERuntimeIREE/Source/IREEDriverRDG/Private/IREEDriverRDGDevice.cpp b/Engine/Plugins/Experimental/NNERuntimeIREE/Source/IREEDriverRDG/Private/IREEDriverRDGDevice.cpp
new file mode 100644
--- /dev/null
+++ b/Engine/Plugins/Experimental/NNERuntimeIREE/Source/IREEDriverRDG/Private/IREEDriverRDGDevice.cpp
@@ -0,0 +1,358 @@
+// Copyright Epic Games, Inc. All Rights Reserved.
+
+#include "IREEDriverRDGDevice.h"
+
+#ifdef WITH_IREE_DRIVER_RDG
+
+#include "IREEDriverRDGCommandBuffer.h"
+#include "IREEDriverRDGDescriptorSetLayout.h"
+#include "IREEDriverRDGDeviceAllocator.h"
+#include "IREEDriverRDGNoOpExecutableCache.h"
+#include "IREEDriverRDGLog.h"
+#include "IREEDriverRDGPipelineLayout.h"
+#include "IREEDriverRDGSemaphore.h"
+
+namespace UE::IREE::HAL::RDG
+{
+
+namespace Private
+{
+
+class FDevice
+{
+public:
+	static iree_status_t Create(iree_allocator_t HostAllocator, iree_hal_device_t** OutDevice)
+	{
+		check(OutDevice);
+		UE_LOG(LogIREEDriverRDG, Display, TEXT("%s"), StringCast<TCHAR>(__FUNCTION__).Get());
+
+		FDevice* Device;
+		IREE_RETURN_IF_ERROR(iree_allocator_malloc(HostAllocator, sizeof(*Device), (void**)& Device));
+		iree_hal_resource_initialize((const void*)&FDevice::VTable, &Device->Resource);
+
+		Device->HostAllocator = HostAllocator;
+		iree_status_t Status = DeviceAllocatorCreate(HostAllocator, &Device->DeviceAllocator);
+		if (!iree_status_is_ok(Status))
+		{
+			iree_allocator_free(HostAllocator, Device);
+			return Status;
+		}
+
+		*OutDevice = (iree_hal_device_t*)Device;
+		return Status;
+	}
+
+	iree_allocator_t GetHostAllocator()
+	{
+		return HostAllocator;
+	}
+
+private:
+	static FDevice* Cast(iree_hal_device_t* Device)
+	{
+		checkf(iree_hal_resource_is(Device, &FDevice::VTable), TEXT("FDevice: type does not match"));
+		return (FDevice*)Device;
+	}
+
+	static void Destroy(iree_hal_device_t *BaseDevice)
+	{
+		UE_LOG(LogIREEDriverRDG, Display, TEXT("%s"), StringCast<TCHAR>(__FUNCTION__).Get());
+		FDevice* Device = Cast(BaseDevice);
+		iree_hal_resource_release(Device->DeviceAllocator);
+		iree_allocator_free(Device->HostAllocator, Device);
+	}
+
+	static iree_string_view_t ID(iree_hal_device_t* Device)
+	{
+		UE_LOG(LogIREEDriverRDG, Display, TEXT("%s"), StringCast<TCHAR>(__FUNCTION__).Get());
+		return iree_make_cstring_view("local-sync");
+	}
+
+	static iree_allocator_t GetHostAllocator(iree_hal_device_t* BaseDevice)
+	{
+		UE_LOG(LogIREEDriverRDG, Display, TEXT("%s"), StringCast<TCHAR>(__FUNCTION__).Get());
+		FDevice* Device = Cast(BaseDevice);
+		return Device->HostAllocator;
+	}
+
+	static iree_hal_allocator_t* GetDeviceAllocator(iree_hal_device_t* BaseDevice)
+	{
+		UE_LOG(LogIREEDriverRDG, Display, TEXT("%s"), StringCast<TCHAR>(__FUNCTION__).Get());
+		FDevice* Device = Cast(BaseDevice);
+		return Device->DeviceAllocator;
+	}
+
+	static void SetDeviceAllocator(iree_hal_device_t* BaseDevice, iree_hal_allocator_t* DeviceAllocator)
+	{
+		UE_LOG(LogIREEDriverRDG, Display, TEXT("%s"), StringCast<TCHAR>(__FUNCTION__).Get());
+
+		Cast(BaseDevice)->DeviceAllocator = DeviceAllocator;
+	}
+
+	static void SetChannelProvider(iree_hal_device_t* BaseDevice, iree_hal_channel_provider_t* ChannelProvider)
+	{
+		UE_LOG(LogIREEDriverRDG, Display, TEXT("%s"), StringCast<TCHAR>(__FUNCTION__).Get());
+		FDevice* Device = Cast(BaseDevice);
+		Device->ChannelProvider = ChannelProvider;
+	}
+
+	static iree_status_t Trim(iree_hal_device_t* Device)
+	{
+		UE_LOG(LogIREEDriverRDG, Display, TEXT("%s"), StringCast<TCHAR>(__FUNCTION__).Get());
+		return iree_ok_status();
+	}
+
+	static iree_status_t QueryInt64(iree_hal_device_t* Device,
+									iree_string_view_t Category,
+									iree_string_view_t Key,
+									int64_t* OutValue)
+	{
+		UE_LOG(LogIREEDriverRDG, Display, TEXT("%s: category='%.*hs' key='%.*hs'"), StringCast<TCHAR>(__FUNCTION__).Get(), (int)Category.size, Category.data, (int)Key.size, Key.data);
+
+		if (iree_string_view_equal(Category, iree_make_cstring_view("hal.executable.format"))) 
+		{
+			*OutValue = 1;
+			return iree_ok_status();
+		}
+		return iree_make_status(IREE_STATUS_NOT_FOUND, 
+								"Unknown device configuration category %.*s :: %.*s", 
+								(int)Category.size, Category.data, (int)Key.size, Key.data);
+	}
+
+	static iree_status_t CreateChannel(iree_hal_device_t* BaseDevice,
+										iree_hal_queue_affinity_t QueueAffinity,
+										iree_hal_channel_params_t Params,
+										iree_hal_channel_t **OutChannel)
+	{
+		UE_LOG(LogIREEDriverRDG, Display, TEXT("%s"), StringCast<TCHAR>(__FUNCTION__).Get());
+		return iree_make_status(IREE_STATUS_UNIMPLEMENTED, __FUNCTION__);
+	}
+
+	static iree_status_t CreateCommandBuffer(iree_hal_device_t* BaseDevice,
+												iree_hal_command_buffer_mode_t Mode,
+												iree_hal_command_category_t CommandCategories,	
+												iree_hal_queue_affinity_t QueueAffinity,
+												iree_host_size_t BindingCapacity,
+												iree_hal_command_buffer_t **OutBuffer)
+	{
+		UE_LOG(LogIREEDriverRDG, Display, TEXT("%s"), StringCast<TCHAR>(__FUNCTION__).Get());
+		FDevice* Device = Cast(BaseDevice);
+		return CommandBufferCreate(Device->HostAllocator, BaseDevice, Mode, CommandCategories, QueueAffinity, BindingCapacity, OutBuffer);
+	}
+
+	static iree_status_t CreateDescriptorSetLayout(iree_hal_device_t* BaseDevice,
+													iree_hal_descriptor_set_layout_flags_t Flags,
+													iree_host_size_t BindingCount,
+													const iree_hal_descriptor_set_layout_binding_t *Bindings,
+													iree_hal_descriptor_set_layout_t **OutDescriptorSetLayout)
+	{
+		UE_LOG(LogIREEDriverRDG, Display, TEXT("%s"), StringCast<TCHAR>(__FUNCTION__).Get());
+
+		FDevice* Device = Cast(BaseDevice);
+		return DescriptorSetLayoutCreate(Device->HostAllocator, OutDescriptorSetLayout);
+	}
+
+	static iree_status_t CreateEvent(iree_hal_device_t* BaseDevice,
+										iree_hal_event_t **OutEvent)
+	{
+		UE_LOG(LogIREEDriverRDG, Display, TEXT("%s"), StringCast<TCHAR>(__FUNCTION__).Get());
+		return iree_make_status(IREE_STATUS_UNIMPLEMENTED, __FUNCTION__);
+	}
+
+	static iree_status_t CreateExecutableCache(iree_hal_device_t* BaseDevice,
+												iree_string_view_t Identifier,
+												iree_loop_t Loop,
+												iree_hal_executable_cache_t **OutExecutableCache)
+	{
+		UE_LOG(LogIREEDriverRDG, Display, TEXT("%s"), StringCast<TCHAR>(__FUNCTION__).Get());
+
+		FDevice* Device = Cast(BaseDevice);
+		return NoOpExecutableCacheCreate(Device->HostAllocator, OutExecutableCache);
+	}
+
+	static iree_status_t ImportFile(iree_hal_device_t* BaseDevice,
+									iree_hal_queue_affinity_t QueueAffinity,
+									iree_hal_memory_access_t Access,
+									iree_io_file_handle_t *Handle,
+									iree_hal_external_file_flags_t Flags,
+									iree_hal_file_t **OutFile)
+	{
+		UE_LOG(LogIREEDriverRDG, Display, TEXT("%s"), StringCast<TCHAR>(__FUNCTION__).Get());
+		return iree_make_status(IREE_STATUS_UNIMPLEMENTED, __FUNCTION__);
+	}
+
+	static iree_status_t CreatePipelineLayout(iree_hal_device_t* BaseDevice,
+												iree_host_size_t PushConstants,
+												iree_host_size_t LayoutCount,
+												iree_hal_descriptor_set_layout_t* const* SetLayouts,
+												iree_hal_pipeline_layout_t **OutPipelineLayout)
+	{
+		UE_LOG(LogIREEDriverRDG, Display, TEXT("%s"), StringCast<TCHAR>(__FUNCTION__).Get());
+
+		FDevice* Device = Cast(BaseDevice);
+		return PipelineLayoutCreate(Device->HostAllocator, OutPipelineLayout);
+	}
+
+	static iree_status_t CreateSemaphore(iree_hal_device_t* BaseDevice,
+											uint64_t InitialValue,
+											iree_hal_semaphore_t **OutSemaphore)
+	{
+		UE_LOG(LogIREEDriverRDG, Display, TEXT("%s"), StringCast<TCHAR>(__FUNCTION__).Get());
+		FDevice* Device = Cast(BaseDevice);
+		return SemaphoreCreate(Device->HostAllocator, InitialValue, OutSemaphore);
+	}
+
+	static iree_hal_semaphore_compatibility_t QuerySemaphoreCompatibility(iree_hal_device_t* BaseDevice, iree_hal_semaphore_t* Semaphore)
+	{
+		UE_LOG(LogIREEDriverRDG, Display, TEXT("%s"), StringCast<TCHAR>(__FUNCTION__).Get());
+		return IREE_HAL_SEMAPHORE_COMPATIBILITY_NONE;
+	}
+
+	static iree_status_t QueueAlloca(iree_hal_device_t* BaseDevice,
+										iree_hal_queue_affinity_t QueueAffinity,
+										const iree_hal_semaphore_list_t WaitSemaphoreList,
+										const iree_hal_semaphore_list_t SignalSemaphoreList,
+										iree_hal_allocator_pool_t Pool,
+										iree_hal_buffer_params_t Params,
+										iree_device_size_t AllocationSize,
+										iree_hal_buffer_t **OutBuffer)
+	{
+		UE_LOG(LogIREEDriverRDG, Display, TEXT("%s"), StringCast<TCHAR>(__FUNCTION__).Get());
+		return iree_hal_allocator_allocate_buffer(iree_hal_device_allocator(BaseDevice), Params, AllocationSize, OutBuffer);
+	}
+
+	static iree_status_t QueueDealloca(iree_hal_device_t* BaseDevice,
+										iree_hal_queue_affinity_t QueueAffinity,
+										const iree_hal_semaphore_list_t WaitSemaphoreList,
+										const iree_hal_semaphore_list_t SignalSemaphoreList,
+										iree_hal_buffer_t *Buffer)
+	{
+		UE_LOG(LogIREEDriverRDG, Display, TEXT("%s"), StringCast<TCHAR>(__FUNCTION__).Get());
+		return iree_make_status(IREE_STATUS_UNIMPLEMENTED, __FUNCTION__);
+	}
+
+	static iree_status_t QueueRead(iree_hal_device_t* BaseDevice,
+									iree_hal_queue_affinity_t QueueAffinity,
+									const iree_hal_semaphore_list_t WaitSemaphoreList,
+									const iree_hal_semaphore_list_t SignalSemaphoreList,
+									iree_hal_file_t* SourceFile,
+									uint64_t SourceOffset,
+									iree_hal_buffer_t* TargetBuffer,
+									iree_device_size_t TargetOffset,
+									iree_device_size_t Length,
+									uint32_t Flags)
+	{
+		UE_LOG(LogIREEDriverRDG, Display, TEXT("%s"), StringCast<TCHAR>(__FUNCTION__).Get());
+		return iree_make_status(IREE_STATUS_UNIMPLEMENTED, __FUNCTION__);
+	}
+
+	static iree_status_t QueueWrite(iree_hal_device_t* BaseDevice,
+									iree_hal_queue_affinity_t QueueAffinity,
+									const iree_hal_semaphore_list_t WaitSemaphoreList,
+									const iree_hal_semaphore_list_t SignalSemaphoreList,
+									iree_hal_buffer_t* SourceBuffer,
+									iree_device_size_t SourceOffset,
+									iree_hal_file_t* TargetFile,
+									uint64_t TargetOffset,
+									iree_device_size_t Length,
+									uint32_t Flags)
+	{
+		UE_LOG(LogIREEDriverRDG, Display, TEXT("%s"), StringCast<TCHAR>(__FUNCTION__).Get());
+		return iree_make_status(IREE_STATUS_UNIMPLEMENTED, __FUNCTION__);
+	}
+	static iree_status_t QueueExecute(iree_hal_device_t* BaseDevice,
+										iree_hal_queue_affinity_t QueueAffinity,
+										const iree_hal_semaphore_list_t WaitSemaphoreList,
+										const iree_hal_semaphore_list_t SignalSemaphoreList,
+										iree_host_size_t CommandBufferCount,
+										iree_hal_command_buffer_t* const* CommandBuffers)
+	{
+		UE_LOG(LogIREEDriverRDG, Display, TEXT("%s"), StringCast<TCHAR>(__FUNCTION__).Get());
+		//return iree_make_status(IREE_STATUS_UNIMPLEMENTED, __FUNCTION__);
+		return iree_ok_status();
+	}
+
+	static iree_status_t QueueFlush(iree_hal_device_t* BaseDevice,
+									iree_hal_queue_affinity_t QueueAffinity)
+	{
+		UE_LOG(LogIREEDriverRDG, Display, TEXT("%s"), StringCast<TCHAR>(__FUNCTION__).Get());
+		return iree_make_status(IREE_STATUS_UNIMPLEMENTED, __FUNCTION__);
+	}
+
+	static iree_status_t WaitSemaphore(iree_hal_device_t* BaseDevice,
+										iree_hal_wait_mode_t WaitMode,
+										const iree_hal_semaphore_list_t SemaphoreList,
+										iree_timeout_t Timeout)
+	{
+		UE_LOG(LogIREEDriverRDG, Display, TEXT("%s"), StringCast<TCHAR>(__FUNCTION__).Get());
+		return iree_make_status(IREE_STATUS_UNIMPLEMENTED, __FUNCTION__);
+	}
+
+	static iree_status_t ProfilingBegin(iree_hal_device_t* BaseDevice,
+										const iree_hal_device_profiling_options_t* Options)
+	{
+		UE_LOG(LogIREEDriverRDG, Display, TEXT("%s"), StringCast<TCHAR>(__FUNCTION__).Get());
+		return iree_make_status(IREE_STATUS_UNIMPLEMENTED, __FUNCTION__);
+	}
+
+	static iree_status_t ProfilingFlush(iree_hal_device_t* BaseDevice)
+	{
+		UE_LOG(LogIREEDriverRDG, Display, TEXT("%s"), StringCast<TCHAR>(__FUNCTION__).Get());
+		return iree_make_status(IREE_STATUS_UNIMPLEMENTED, __FUNCTION__);
+	}
+
+	static iree_status_t ProfilingEnd(iree_hal_device_t* BaseDevice)
+	{
+		UE_LOG(LogIREEDriverRDG, Display, TEXT("%s"), StringCast<TCHAR>(__FUNCTION__).Get());
+		return iree_make_status(IREE_STATUS_UNIMPLEMENTED, __FUNCTION__);
+	}
+
+	static const iree_hal_device_vtable_t VTable;
+
+	iree_hal_resource_t Resource;
+	iree_allocator_t HostAllocator;
+	iree_hal_allocator_t* DeviceAllocator;
+	iree_hal_channel_provider_t* ChannelProvider;
+};
+
+const iree_hal_device_vtable_t FDevice::VTable =
+{
+	.destroy = FDevice::Destroy,
+	.id = FDevice::ID,
+	.host_allocator = FDevice::GetHostAllocator,
+	.device_allocator = FDevice::GetDeviceAllocator,
+	.replace_device_allocator = FDevice::SetDeviceAllocator,
+	.replace_channel_provider = FDevice::SetChannelProvider,
+	.trim = FDevice::Trim,
+	.query_i64 = FDevice::QueryInt64,
+	.create_channel = FDevice::CreateChannel,
+	.create_command_buffer = FDevice::CreateCommandBuffer,
+	.create_descriptor_set_layout = FDevice::CreateDescriptorSetLayout,
+	.create_event = FDevice::CreateEvent,
+	.create_executable_cache = FDevice::CreateExecutableCache,
+	.import_file = FDevice::ImportFile,
+	.create_pipeline_layout = FDevice::CreatePipelineLayout,
+	.create_semaphore = FDevice::CreateSemaphore,
+	.query_semaphore_compatibility = FDevice::QuerySemaphoreCompatibility,
+	.queue_alloca = FDevice::QueueAlloca,
+	.queue_dealloca = FDevice::QueueDealloca,
+	.queue_read = FDevice::QueueRead,
+	.queue_write = FDevice::QueueWrite,
+	.queue_execute = FDevice::QueueExecute,
+	.queue_flush = FDevice::QueueFlush,
+	.wait_semaphores = FDevice::WaitSemaphore,
+	.profiling_begin = FDevice::ProfilingBegin,
+	.profiling_flush = FDevice::ProfilingFlush,
+	.profiling_end = FDevice::ProfilingEnd
+};
+
+} // namespace Private
+
+iree_status_t DeviceCreate(iree_allocator_t HostAllocator, iree_hal_device_t** OutDevice)
+{
+	return Private::FDevice::Create(HostAllocator, OutDevice);
+}
+
+} // UE::IREE
+
+#endif // WITH_IREE_DRIVER_RDG
\ No newline at end of file
diff --git a/Engine/Plugins/Experimental/NNERuntimeIREE/Source/IREEDriverRDG/Private/IREEDriverRDGDevice.h b/Engine/Plugins/Experimental/NNERuntimeIREE/Source/IREEDriverRDG/Private/IREEDriverRDGDevice.h
new file mode 100644
--- /dev/null
+++ b/Engine/Plugins/Experimental/NNERuntimeIREE/Source/IREEDriverRDG/Private/IREEDriverRDGDevice.h
@@ -0,0 +1,9 @@
+// Copyright Epic Games, Inc. All Rights Reserved.
+
+#pragma once
+
+#ifdef WITH_IREE_DRIVER_RDG
+
+#include "IREEDriverRDG.h"
+
+#endif // WITH_IREE_DRIVER_RDG
\ No newline at end of file
diff --git a/Engine/Plugins/Experimental/NNERuntimeIREE/Source/IREEDriverRDG/Private/IREEDriverRDGDeviceAllocator.cpp b/Engine/Plugins/Experimental/NNERuntimeIREE/Source/IREEDriverRDG/Private/IREEDriverRDGDeviceAllocator.cpp
new file mode 100644
--- /dev/null
+++ b/Engine/Plugins/Experimental/NNERuntimeIREE/Source/IREEDriverRDG/Private/IREEDriverRDGDeviceAllocator.cpp
@@ -0,0 +1,189 @@
+// Copyright Epic Games, Inc. All Rights Reserved.
+
+#include "IREEDriverRDGDeviceAllocator.h"
+
+#ifdef WITH_IREE_DRIVER_RDG
+
+#include "IREEDriverRDGBuffer.h"
+#include "IREEDriverRDGLog.h"
+
+namespace UE::IREE::HAL::RDG
+{
+
+namespace Private
+{
+
+class FDeviceAllocator
+{
+public:
+	static iree_status_t Create(iree_allocator_t HostAllocator, iree_hal_allocator_t** OutDeviceAllocator)
+	{
+		check(OutDeviceAllocator);
+		UE_LOG(LogIREEDriverRDG, Display, TEXT("%s"), StringCast<TCHAR>(__FUNCTION__).Get());
+
+		FDeviceAllocator* DeviceAllocator;
+		IREE_RETURN_IF_ERROR(iree_allocator_malloc(HostAllocator, sizeof(*DeviceAllocator), (void**)&DeviceAllocator));
+		iree_hal_resource_initialize((const void*)&FDeviceAllocator::VTable, &DeviceAllocator->Resource);
+		DeviceAllocator->HostAllocator = HostAllocator;
+
+		*OutDeviceAllocator = (iree_hal_allocator_t*)DeviceAllocator;
+		return iree_ok_status();
+	}
+
+	void SetGraphBuilder(FRDGBuilder& InGraphBuilder)
+	{
+		check(GraphBuilder == nullptr);
+		GraphBuilder = &InGraphBuilder;
+	}
+
+	FRDGBuilder& GetGraphBuilder() const
+	{
+		check(GraphBuilder);
+		return *GraphBuilder;
+	}
+
+	void ResetGraphBuilder()
+	{
+		check(GraphBuilder);
+		GraphBuilder = nullptr;
+	}
+
+private:
+	static FDeviceAllocator* Cast(const iree_hal_allocator_t* Allocator)
+	{
+		checkf(iree_hal_resource_is(Allocator, &FDeviceAllocator::VTable), TEXT("FDeviceAllocator: type does not match"));
+		return (FDeviceAllocator*)Allocator;
+	}
+
+	static void Destroy(iree_hal_allocator_t *BaseAllocator)
+	{
+		UE_LOG(LogIREEDriverRDG, Display, TEXT("%s"), StringCast<TCHAR>(__FUNCTION__).Get());
+		FDeviceAllocator* Allocator = Cast(BaseAllocator);
+		iree_allocator_free(Allocator->HostAllocator, Allocator);
+	}
+
+	static iree_allocator_t GetHostAllocator(const iree_hal_allocator_t* BaseAllocator)
+	{
+		UE_LOG(LogIREEDriverRDG, Display, TEXT("%s"), StringCast<TCHAR>(__FUNCTION__).Get());
+		FDeviceAllocator* Allocator = Cast(BaseAllocator);
+		return Allocator->HostAllocator;
+	}
+
+	static iree_status_t Trim(iree_hal_allocator_t* BaseAllocator)
+	{
+		UE_LOG(LogIREEDriverRDG, Display, TEXT("%s"), StringCast<TCHAR>(__FUNCTION__).Get());
+		return iree_make_status(IREE_STATUS_UNIMPLEMENTED, __FUNCTION__);
+	}
+
+	static void QueryStatistics(iree_hal_allocator_t* BaseAllocator,
+								iree_hal_allocator_statistics_t* OutStatistics)
+	{
+		UE_LOG(LogIREEDriverRDG, Display, TEXT("%s"), StringCast<TCHAR>(__FUNCTION__).Get());
+	}
+
+	static iree_status_t QueryMemoryHeaps(iree_hal_allocator_t* BaseAllocator,
+											iree_host_size_t Capacity,
+											iree_hal_allocator_memory_heap_t* Heaps,
+											iree_host_size_t* OutCount)
+	{
+		UE_LOG(LogIREEDriverRDG, Display, TEXT("%s"), StringCast<TCHAR>(__FUNCTION__).Get());
+		return iree_make_status(IREE_STATUS_UNIMPLEMENTED, __FUNCTION__);
+	}
+
+	static iree_hal_buffer_compatibility_t QueryBufferCompatibility(iree_hal_allocator_t* BaseAllocator,
+																	iree_hal_buffer_params_t* Params,
+																	iree_device_size_t* AllocationSize)
+	{
+		check(BaseAllocator);
+		check(Params);
+		check(AllocationSize);
+		UE_LOG(LogIREEDriverRDG, Display, TEXT("%s allocationSize %lu"), StringCast<TCHAR>(__FUNCTION__).Get(), *AllocationSize);
+
+		return IREE_HAL_BUFFER_COMPATIBILITY_IMPORTABLE | 
+				IREE_HAL_BUFFER_COMPATIBILITY_EXPORTABLE |
+				IREE_HAL_BUFFER_COMPATIBILITY_QUEUE_TRANSFER |
+				IREE_HAL_BUFFER_COMPATIBILITY_QUEUE_DISPATCH;
+	}
+
+	static iree_status_t AllocateBuffer(iree_hal_allocator_t* BaseAllocator,
+										const iree_hal_buffer_params_t* Params,
+										iree_device_size_t AllocationSize,
+										iree_hal_buffer_t** OutBuffer)
+	{
+		UE_LOG(LogIREEDriverRDG, Display, TEXT("%s"), StringCast<TCHAR>(__FUNCTION__).Get());
+		FDeviceAllocator* Allocator = Cast(BaseAllocator);
+		return BufferCreate(Allocator->HostAllocator, BaseAllocator, Params, AllocationSize, OutBuffer);
+	}
+
+	static void DeallocateBuffer(iree_hal_allocator_t* BaseAllocator,
+											iree_hal_buffer_t* Buffer)
+	{
+		UE_LOG(LogIREEDriverRDG, Display, TEXT("%s"), StringCast<TCHAR>(__FUNCTION__).Get());
+	}
+
+	static iree_status_t ImportBuffer(iree_hal_allocator_t* BaseAllocator,
+										const iree_hal_buffer_params_t* Params,
+										iree_hal_external_buffer_t* ExternalBuffer,
+										iree_hal_buffer_release_callback_t ReleaseCallback,
+										iree_hal_buffer_t** OutBuffer)
+	{
+		UE_LOG(LogIREEDriverRDG, Display, TEXT("%s"), StringCast<TCHAR>(__FUNCTION__).Get());
+		return iree_make_status(IREE_STATUS_UNIMPLEMENTED, __FUNCTION__);
+	}
+
+	static iree_status_t ExportBuffer(iree_hal_allocator_t* BaseAllocator,
+										iree_hal_buffer_t* Buffer,
+										iree_hal_external_buffer_type_t RequestedType,
+										iree_hal_external_buffer_flags_t RequestedFlags,
+										iree_hal_external_buffer_t* OutExternalBuffer)
+	{
+		UE_LOG(LogIREEDriverRDG, Display, TEXT("%s"), StringCast<TCHAR>(__FUNCTION__).Get());
+		return iree_make_status(IREE_STATUS_UNIMPLEMENTED, __FUNCTION__);
+	}
+	
+	static const iree_hal_allocator_vtable_t VTable;
+
+	iree_hal_resource_t Resource;
+	iree_allocator_t HostAllocator;
+	FRDGBuilder* GraphBuilder = nullptr;
+};
+
+const iree_hal_allocator_vtable_t FDeviceAllocator::VTable =
+{
+	.destroy = FDeviceAllocator::Destroy,
+	.host_allocator = FDeviceAllocator::GetHostAllocator,
+	.trim = FDeviceAllocator::Trim,
+	.query_statistics = FDeviceAllocator::QueryStatistics,
+	.query_memory_heaps = FDeviceAllocator::QueryMemoryHeaps,
+	.query_buffer_compatibility = FDeviceAllocator::QueryBufferCompatibility,
+	.allocate_buffer = FDeviceAllocator::AllocateBuffer,
+	.deallocate_buffer = FDeviceAllocator::DeallocateBuffer,
+	.import_buffer = FDeviceAllocator::ImportBuffer,
+	.export_buffer = FDeviceAllocator::ExportBuffer
+};
+
+} // namespace Private
+
+iree_status_t DeviceAllocatorCreate(iree_allocator_t HostAllocator, iree_hal_allocator_t** OutDeviceAllocator)
+{
+	return Private::FDeviceAllocator::Create(HostAllocator, OutDeviceAllocator);
+}
+
+void DeviceAllocatorSetGraphBuilder(iree_hal_allocator_t* DeviceAllocator, FRDGBuilder& GraphBuilder)
+{
+	((Private::FDeviceAllocator*)DeviceAllocator)->SetGraphBuilder(GraphBuilder);
+}
+
+FRDGBuilder& DeviceAllocatorGetGraphBuilder(iree_hal_allocator_t* DeviceAllocator)
+{
+	return ((Private::FDeviceAllocator*)DeviceAllocator)->GetGraphBuilder();
+}
+
+void DeviceAllocatorResetGraphBuilder(iree_hal_allocator_t* DeviceAllocator)
+{
+	((Private::FDeviceAllocator*)DeviceAllocator)->ResetGraphBuilder();
+}
+
+} // UE::IREE
+
+#endif // WITH_IREE_DRIVER_RDG
\ No newline at end of file
diff --git a/Engine/Plugins/Experimental/NNERuntimeIREE/Source/IREEDriverRDG/Private/IREEDriverRDGDeviceAllocator.h b/Engine/Plugins/Experimental/NNERuntimeIREE/Source/IREEDriverRDG/Private/IREEDriverRDGDeviceAllocator.h
new file mode 100644
--- /dev/null
+++ b/Engine/Plugins/Experimental/NNERuntimeIREE/Source/IREEDriverRDG/Private/IREEDriverRDGDeviceAllocator.h
@@ -0,0 +1,16 @@
+// Copyright Epic Games, Inc. All Rights Reserved.
+
+#pragma once
+
+#ifdef WITH_IREE_DRIVER_RDG
+
+#include "IREEDriverRDG.h"
+
+namespace UE::IREE::HAL::RDG
+{
+
+iree_status_t DeviceAllocatorCreate(iree_allocator_t HostAllocator, iree_hal_allocator_t** OutDeviceAllocator);
+
+} // UE::IREE
+
+#endif // WITH_IREE_DRIVER_RDG
\ No newline at end of file
diff --git a/Engine/Plugins/Experimental/NNERuntimeIREE/Source/IREEDriverRDG/Private/IREEDriverRDGExecutable.cpp b/Engine/Plugins/Experimental/NNERuntimeIREE/Source/IREEDriverRDG/Private/IREEDriverRDGExecutable.cpp
new file mode 100644
--- /dev/null
+++ b/Engine/Plugins/Experimental/NNERuntimeIREE/Source/IREEDriverRDG/Private/IREEDriverRDGExecutable.cpp
@@ -0,0 +1,70 @@
+// Copyright Epic Games, Inc. All Rights Reserved.
+
+#include "IREEDriverRDGExecutable.h"
+
+#ifdef WITH_IREE_DRIVER_RDG
+
+#include "IREEDriverRDGLog.h"
+
+namespace UE::IREE::HAL::RDG
+{
+
+namespace Private
+{
+
+class FExecutable
+{
+public:
+	static iree_status_t Create(iree_allocator_t HostAllocator, iree_hal_executable_t** OutExecutable)
+	{
+		check(OutExecutable);
+		UE_LOG(LogIREEDriverRDG, Display, TEXT("%s"), StringCast<TCHAR>(__FUNCTION__).Get());
+
+		FExecutable* Executable;
+		IREE_RETURN_IF_ERROR(iree_allocator_malloc(HostAllocator, sizeof(*Executable), (void**)&Executable));
+		iree_hal_resource_initialize((const void*)&FExecutable::VTable, &Executable->Resource);
+		Executable->HostAllocator = HostAllocator;
+
+		*OutExecutable = (iree_hal_executable_t*)Executable;
+		return iree_ok_status();
+	}
+private:
+	static FExecutable* Cast(iree_hal_executable_t* Executable)
+	{
+		checkf(iree_hal_resource_is(Executable, &FExecutable::VTable), TEXT("FExecutable: type does not match"));
+		return (FExecutable*)Executable;
+	}
+
+	static void Destroy(iree_hal_executable_t* BaseExecutable)
+	{
+		UE_LOG(LogIREEDriverRDG, Display, TEXT("%s"), StringCast<TCHAR>(__FUNCTION__).Get());
+		
+		FExecutable* Executable = Cast(BaseExecutable);
+		iree_allocator_free(Executable->HostAllocator, Executable);
+	}
+
+	inline static const iree_hal_executable_vtable_t VTable = []
+	{
+		iree_hal_executable_vtable_t Result =
+		{
+			.destroy = FExecutable::Destroy,
+		};
+		return Result;
+
+	}();
+
+private:
+	iree_hal_resource_t Resource;
+	iree_allocator_t HostAllocator;
+};
+
+} // namespace Private
+
+iree_status_t ExecutableCreate(iree_allocator_t HostAllocator, iree_hal_executable_t** OutExecutable)
+{
+	return Private::FExecutable::Create(HostAllocator, OutExecutable);
+}
+
+} // UE::IREE
+
+#endif // WITH_IREE_DRIVER_RDG
\ No newline at end of file
diff --git a/Engine/Plugins/Experimental/NNERuntimeIREE/Source/IREEDriverRDG/Private/IREEDriverRDGExecutable.h b/Engine/Plugins/Experimental/NNERuntimeIREE/Source/IREEDriverRDG/Private/IREEDriverRDGExecutable.h
new file mode 100644
--- /dev/null
+++ b/Engine/Plugins/Experimental/NNERuntimeIREE/Source/IREEDriverRDG/Private/IREEDriverRDGExecutable.h
@@ -0,0 +1,16 @@
+// Copyright Epic Games, Inc. All Rights Reserved.
+
+#pragma once
+
+#ifdef WITH_IREE_DRIVER_RDG
+
+#include "IREEDriverRDG.h"
+
+namespace UE::IREE::HAL::RDG
+{
+
+iree_status_t ExecutableCreate(iree_allocator_t HostAllocator, iree_hal_executable_t** OutExecutable);
+
+} // UE::IREE
+
+#endif // WITH_IREE_DRIVER_RDG
\ No newline at end of file
diff --git a/Engine/Plugins/Experimental/NNERuntimeIREE/Source/IREEDriverRDG/Private/IREEDriverRDGModule.cpp b/Engine/Plugins/Experimental/NNERuntimeIREE/Source/IREEDriverRDG/Private/IREEDriverRDGModule.cpp
new file mode 100644
--- /dev/null
+++ b/Engine/Plugins/Experimental/NNERuntimeIREE/Source/IREEDriverRDG/Private/IREEDriverRDGModule.cpp
@@ -0,0 +1,20 @@
+// Copyright Epic Games, Inc. All Rights Reserved.
+
+#include "IREEDriverRDGModule.h"
+
+#include "Misc/CoreDelegates.h"
+#include "IREEDriverRDGLog.h"
+
+DEFINE_LOG_CATEGORY(LogIREEDriverRDG);
+
+void FIREEDriverRDGModule::StartupModule()
+{
+
+}
+
+void FIREEDriverRDGModule::ShutdownModule()
+{
+
+}
+
+IMPLEMENT_MODULE(FIREEDriverRDGModule, IREEDriverRDG)
\ No newline at end of file
diff --git a/Engine/Plugins/Experimental/NNERuntimeIREE/Source/IREEDriverRDG/Private/IREEDriverRDGModule.h b/Engine/Plugins/Experimental/NNERuntimeIREE/Source/IREEDriverRDG/Private/IREEDriverRDGModule.h
new file mode 100644
--- /dev/null
+++ b/Engine/Plugins/Experimental/NNERuntimeIREE/Source/IREEDriverRDG/Private/IREEDriverRDGModule.h
@@ -0,0 +1,13 @@
+// Copyright Epic Games, Inc. All Rights Reserved.
+
+#pragma once
+
+#include "Modules/ModuleManager.h"
+
+class FIREEDriverRDGModule : public IModuleInterface
+{
+public:
+	/** IModuleInterface implementation */
+	virtual void StartupModule() override;
+	virtual void ShutdownModule() override;
+};
diff --git a/Engine/Plugins/Experimental/NNERuntimeIREE/Source/IREEDriverRDG/Private/IREEDriverRDGNoOpExecutableCache.cpp b/Engine/Plugins/Experimental/NNERuntimeIREE/Source/IREEDriverRDG/Private/IREEDriverRDGNoOpExecutableCache.cpp
new file mode 100644
--- /dev/null
+++ b/Engine/Plugins/Experimental/NNERuntimeIREE/Source/IREEDriverRDG/Private/IREEDriverRDGNoOpExecutableCache.cpp
@@ -0,0 +1,82 @@
+// Copyright Epic Games, Inc. All Rights Reserved.
+
+#include "IREEDriverRDGNoOpExecutableCache.h"
+
+#ifdef WITH_IREE_DRIVER_RDG
+
+#include "IREEDriverRDGExecutable.h"
+#include "IREEDriverRDGLog.h"
+
+namespace UE::IREE::HAL::RDG
+{
+
+namespace Private
+{
+
+class FNoOpExecutableCache
+{
+public:
+	static iree_status_t Create(iree_allocator_t HostAllocator, iree_hal_executable_cache_t** OutExecutableCache)
+	{
+		check(OutExecutableCache);
+		UE_LOG(LogIREEDriverRDG, Display, TEXT("%s"), StringCast<TCHAR>(__FUNCTION__).Get());
+
+		FNoOpExecutableCache* ExecutableCache;
+		IREE_RETURN_IF_ERROR(iree_allocator_malloc(HostAllocator, sizeof(*ExecutableCache), (void**)&ExecutableCache));
+		iree_hal_resource_initialize((const void*)&FNoOpExecutableCache::VTable, &ExecutableCache->Resource);
+		ExecutableCache->HostAllocator = HostAllocator;
+
+		*OutExecutableCache = (iree_hal_executable_cache_t*)ExecutableCache;
+		return iree_ok_status();
+	}
+private:
+	static FNoOpExecutableCache* Cast(iree_hal_executable_cache_t* ExecutableCache)
+	{
+		checkf(iree_hal_resource_is(ExecutableCache, &FNoOpExecutableCache::VTable), TEXT("FNoOpExecutableCache: type does not match"));
+		return (FNoOpExecutableCache*)ExecutableCache;
+	}
+
+	static void Destroy(iree_hal_executable_cache_t* BaseExecutableCache)
+	{
+		UE_LOG(LogIREEDriverRDG, Display, TEXT("%s"), StringCast<TCHAR>(__FUNCTION__).Get());
+
+		FNoOpExecutableCache* ExecutableCache = Cast(BaseExecutableCache);
+		iree_allocator_free(ExecutableCache->HostAllocator, ExecutableCache);
+	}
+
+	static bool CanPrepareFormat(iree_hal_executable_cache_t* BaseExecutableCache, iree_hal_executable_caching_mode_t CachingMode, iree_string_view_t ExecutableFormat) 
+	{
+		UE_LOG(LogIREEDriverRDG, Display, TEXT("%s format: %*.hs"), StringCast<TCHAR>(__FUNCTION__).Get(), (int)ExecutableFormat.size, ExecutableFormat.data);
+		return true;
+	}
+
+	static iree_status_t PrepareExecutable(iree_hal_executable_cache_t* BaseExecutableCache, const iree_hal_executable_params_t* ExecutableParams, iree_hal_executable_t** OutExecutable) 
+	{
+		UE_LOG(LogIREEDriverRDG, Display, TEXT("%s"), StringCast<TCHAR>(__FUNCTION__).Get());
+
+		return ExecutableCreate(Cast(BaseExecutableCache)->HostAllocator, OutExecutable);
+	}
+
+	static const iree_hal_executable_cache_vtable_t VTable;
+
+	iree_hal_resource_t Resource;
+	iree_allocator_t HostAllocator;
+};
+
+const iree_hal_executable_cache_vtable_t FNoOpExecutableCache::VTable = 
+{
+	.destroy = FNoOpExecutableCache::Destroy,
+	.can_prepare_format = FNoOpExecutableCache::CanPrepareFormat,
+	.prepare_executable = FNoOpExecutableCache::PrepareExecutable
+};
+
+} // namespace Private
+
+iree_status_t NoOpExecutableCacheCreate(iree_allocator_t HostAllocator, iree_hal_executable_cache_t** OutExecutableCache)
+{
+	return Private::FNoOpExecutableCache::Create(HostAllocator, OutExecutableCache);
+}
+
+} // UE::IREE
+
+#endif // WITH_IREE_DRIVER_RDG
\ No newline at end of file
diff --git a/Engine/Plugins/Experimental/NNERuntimeIREE/Source/IREEDriverRDG/Private/IREEDriverRDGNoOpExecutableCache.h b/Engine/Plugins/Experimental/NNERuntimeIREE/Source/IREEDriverRDG/Private/IREEDriverRDGNoOpExecutableCache.h
new file mode 100644
--- /dev/null
+++ b/Engine/Plugins/Experimental/NNERuntimeIREE/Source/IREEDriverRDG/Private/IREEDriverRDGNoOpExecutableCache.h
@@ -0,0 +1,16 @@
+// Copyright Epic Games, Inc. All Rights Reserved.
+
+#pragma once
+
+#ifdef WITH_IREE_DRIVER_RDG
+
+#include "IREEDriverRDG.h"
+
+namespace UE::IREE::HAL::RDG
+{
+
+iree_status_t NoOpExecutableCacheCreate(iree_allocator_t HostAllocator, iree_hal_executable_cache_t** OutExecutableCache);
+
+} // UE::IREE
+
+#endif // WITH_IREE_DRIVER_RDG
\ No newline at end of file
diff --git a/Engine/Plugins/Experimental/NNERuntimeIREE/Source/IREEDriverRDG/Private/IREEDriverRDGPipelineLayout.cpp b/Engine/Plugins/Experimental/NNERuntimeIREE/Source/IREEDriverRDG/Private/IREEDriverRDGPipelineLayout.cpp
new file mode 100644
--- /dev/null
+++ b/Engine/Plugins/Experimental/NNERuntimeIREE/Source/IREEDriverRDG/Private/IREEDriverRDGPipelineLayout.cpp
@@ -0,0 +1,65 @@
+// Copyright Epic Games, Inc. All Rights Reserved.
+
+#include "IREEDriverRDGPipelineLayout.h"
+
+#ifdef WITH_IREE_DRIVER_RDG
+
+#include "IREEDriverRDGLog.h"
+
+namespace UE::IREE::HAL::RDG
+{
+
+namespace Private
+{
+
+class FPipelineLayout
+{
+private:
+	iree_hal_resource_t Resource;
+	iree_allocator_t HostAllocator;
+public:
+	static iree_status_t Create(iree_allocator_t HostAllocator, iree_hal_pipeline_layout_t** OutPipelineLayout)
+	{
+		check(OutPipelineLayout);
+		UE_LOG(LogIREEDriverRDG, Display, TEXT("%s"), StringCast<TCHAR>(__FUNCTION__).Get());
+
+		FPipelineLayout* PipelineLayout;
+		IREE_RETURN_IF_ERROR(iree_allocator_malloc(HostAllocator, sizeof(*PipelineLayout), (void**)&PipelineLayout));
+		iree_hal_resource_initialize((const void*)&FPipelineLayout::VTable, &PipelineLayout->Resource);
+		PipelineLayout->HostAllocator = HostAllocator;
+
+		*OutPipelineLayout = (iree_hal_pipeline_layout_t*)PipelineLayout;
+		return iree_ok_status();
+	}
+private:
+	static FPipelineLayout* Cast(iree_hal_pipeline_layout_t* PipelineLayout)
+	{
+		checkf(iree_hal_resource_is(PipelineLayout, &FPipelineLayout::VTable), TEXT("FPipelineLayout: type does not match"));
+		return (FPipelineLayout*)PipelineLayout;
+	}
+
+	static void Destroy(iree_hal_pipeline_layout_t* BasePipelineLayout)
+	{
+		UE_LOG(LogIREEDriverRDG, Display, TEXT("%s"), StringCast<TCHAR>(__FUNCTION__).Get());
+		FPipelineLayout* PipelineLayout = Cast(BasePipelineLayout);
+		iree_allocator_free(PipelineLayout->HostAllocator, PipelineLayout);
+	}
+
+	static const iree_hal_pipeline_layout_vtable_t VTable;
+};
+
+const iree_hal_pipeline_layout_vtable_t FPipelineLayout::VTable = 
+{
+	.destroy = FPipelineLayout::Destroy,
+};
+
+} // namespace Private
+
+iree_status_t PipelineLayoutCreate(iree_allocator_t HostAllocator, iree_hal_pipeline_layout_t** OutPipelineLayout)
+{
+	return Private::FPipelineLayout::Create(HostAllocator, OutPipelineLayout);
+}
+
+} // UE::IREE
+
+#endif // WITH_IREE_DRIVER_RDG
\ No newline at end of file
diff --git a/Engine/Plugins/Experimental/NNERuntimeIREE/Source/IREEDriverRDG/Private/IREEDriverRDGPipelineLayout.h b/Engine/Plugins/Experimental/NNERuntimeIREE/Source/IREEDriverRDG/Private/IREEDriverRDGPipelineLayout.h
new file mode 100644
--- /dev/null
+++ b/Engine/Plugins/Experimental/NNERuntimeIREE/Source/IREEDriverRDG/Private/IREEDriverRDGPipelineLayout.h
@@ -0,0 +1,16 @@
+// Copyright Epic Games, Inc. All Rights Reserved.
+
+#pragma once
+
+#ifdef WITH_IREE_DRIVER_RDG
+
+#include "IREEDriverRDG.h"
+
+namespace UE::IREE::HAL::RDG
+{
+
+iree_status_t PipelineLayoutCreate(iree_allocator_t HostAllocator, iree_hal_pipeline_layout_t** OutPipelineLayout);
+
+} // UE::IREE
+
+#endif // WITH_IREE_DRIVER_RDG
\ No newline at end of file
diff --git a/Engine/Plugins/Experimental/NNERuntimeIREE/Source/IREEDriverRDG/Private/IREEDriverRDGSemaphore.cpp b/Engine/Plugins/Experimental/NNERuntimeIREE/Source/IREEDriverRDG/Private/IREEDriverRDGSemaphore.cpp
new file mode 100644
--- /dev/null
+++ b/Engine/Plugins/Experimental/NNERuntimeIREE/Source/IREEDriverRDG/Private/IREEDriverRDGSemaphore.cpp
@@ -0,0 +1,110 @@
+// Copyright Epic Games, Inc. All Rights Reserved.
+
+#include "IREEDriverRDGSemaphore.h"
+
+#ifdef WITH_IREE_DRIVER_RDG
+
+#include "HAL/Platform.h"
+
+#if PLATFORM_MICROSOFT
+#include "Microsoft/AllowMicrosoftPlatformTypes.h"
+#include "Microsoft/AllowMicrosoftPlatformAtomics.h"
+#endif // PLATFORM_MICROSOFT
+THIRD_PARTY_INCLUDES_START
+#include "IREEDriverRDGLog.h"
+THIRD_PARTY_INCLUDES_END
+#if PLATFORM_MICROSOFT
+#include "Microsoft/HideMicrosoftPlatformAtomics.h"
+#include "Microsoft/HideMicrosoftPlatformTypes.h"
+#endif // PLATFORM_MICROSOFT
+
+#include "iree/hal/utils/semaphore_base.h"
+
+namespace UE::IREE::HAL::RDG
+{
+
+namespace Private
+{
+
+class FSemaphore
+{
+private:
+	iree_hal_semaphore_t Base;
+	iree_allocator_t HostAllocator;
+public:
+	static iree_status_t Create(iree_allocator_t HostAllocator, uint64 InitialValue, iree_hal_semaphore_t** OutSemaphore)
+	{
+		check(OutSemaphore);
+		UE_LOG(LogIREEDriverRDG, Display, TEXT("%s InitialValue %d"), StringCast<TCHAR>(__FUNCTION__).Get(), InitialValue);
+
+		FSemaphore* Semaphore;
+		IREE_RETURN_IF_ERROR(iree_allocator_malloc(HostAllocator, sizeof(*Semaphore), (void**)&Semaphore));
+		iree_hal_semaphore_initialize(&FSemaphore::VTable, &Semaphore->Base);
+		Semaphore->HostAllocator = HostAllocator;
+
+		*OutSemaphore = (iree_hal_semaphore_t*)Semaphore;
+		return iree_ok_status();
+	}
+private:
+	static FSemaphore* Cast(iree_hal_semaphore_t* Semaphore)
+	{
+		checkf(iree_hal_resource_is(Semaphore, &FSemaphore::VTable), TEXT("FSemaphore: type does not match"));
+		return (FSemaphore*)Semaphore;
+	}
+
+	static void Destroy(iree_hal_semaphore_t* BaseSemaphore)
+	{
+		UE_LOG(LogIREEDriverRDG, Display, TEXT("%s"), StringCast<TCHAR>(__FUNCTION__).Get());
+		FSemaphore* Semaphore = Cast(BaseSemaphore);
+		iree_hal_semaphore_deinitialize(&Semaphore->Base);
+		iree_allocator_free(Semaphore->HostAllocator, Semaphore);
+	}
+
+	static iree_status_t Query(iree_hal_semaphore_t* BaseSemaphore, uint64_t* OutValue)
+	{
+		UE_LOG(LogIREEDriverRDG, Display, TEXT("%s"), StringCast<TCHAR>(__FUNCTION__).Get());
+		// return iree_make_status(IREE_STATUS_UNIMPLEMENTED, __FUNCTION__);
+		return iree_ok_status();
+	}
+
+	static iree_status_t Signal(iree_hal_semaphore_t* BaseSemaphore, uint64_t NewValue)
+	{
+		UE_LOG(LogIREEDriverRDG, Display, TEXT("%s"), StringCast<TCHAR>(__FUNCTION__).Get());
+		// return iree_make_status(IREE_STATUS_UNIMPLEMENTED, __FUNCTION__);
+		return iree_ok_status();
+	}
+
+	static void Fail(iree_hal_semaphore_t* BaseSemaphore, iree_status_t Status)
+	{
+		UE_LOG(LogIREEDriverRDG, Display, TEXT("%s"), StringCast<TCHAR>(__FUNCTION__).Get());
+	}
+
+	static iree_status_t Wait(iree_hal_semaphore_t* BaseSemaphore, uint64_t Value, iree_timeout_t Timeout)
+	{
+		UE_LOG(LogIREEDriverRDG, Display, TEXT("%s"), StringCast<TCHAR>(__FUNCTION__).Get());
+		//return iree_make_status(IREE_STATUS_UNIMPLEMENTED, __FUNCTION__);
+		return iree_ok_status();
+	}
+
+	static const iree_hal_semaphore_vtable_t VTable;
+};
+
+const iree_hal_semaphore_vtable_t FSemaphore::VTable =
+{
+	.destroy = FSemaphore::Destroy,
+	.query = FSemaphore::Query,
+	.signal = FSemaphore::Signal,
+	.fail = FSemaphore::Fail,
+	.wait = FSemaphore::Wait
+};
+
+} // namespace Private
+
+iree_status_t SemaphoreCreate(iree_allocator_t HostAllocator, uint64 InitialValue, iree_hal_semaphore_t** OutSemaphore)
+{
+	return Private::FSemaphore::Create(HostAllocator, InitialValue, OutSemaphore);
+}
+
+} // UE::IREE
+
+#endif // WITH_IREE_DRIVER_RDG
\ No newline at end of file
diff --git a/Engine/Plugins/Experimental/NNERuntimeIREE/Source/IREEDriverRDG/Private/IREEDriverRDGSemaphore.h b/Engine/Plugins/Experimental/NNERuntimeIREE/Source/IREEDriverRDG/Private/IREEDriverRDGSemaphore.h
new file mode 100644
--- /dev/null
+++ b/Engine/Plugins/Experimental/NNERuntimeIREE/Source/IREEDriverRDG/Private/IREEDriverRDGSemaphore.h
@@ -0,0 +1,16 @@
+// Copyright Epic Games, Inc. All Rights Reserved.
+
+#pragma once
+
+#ifdef WITH_IREE_DRIVER_RDG
+
+#include "IREEDriverRDG.h"
+
+namespace UE::IREE::HAL::RDG
+{
+
+iree_status_t SemaphoreCreate(iree_allocator_t HostAllocator, uint64 InitialValue, iree_hal_semaphore_t** OutSemaphore);
+
+} // UE::IREE
+
+#endif // WITH_IREE_DRIVER_RDG
\ No newline at end of file
diff --git a/Engine/Plugins/Experimental/NNERuntimeIREE/Source/IREEDriverRDGShaders/IREEDriverRDGShaders.Build.cs b/Engine/Plugins/Experimental/NNERuntimeIREE/Source/IREEDriverRDGShaders/IREEDriverRDGShaders.Build.cs
new file mode 100644
--- /dev/null
+++ b/Engine/Plugins/Experimental/NNERuntimeIREE/Source/IREEDriverRDGShaders/IREEDriverRDGShaders.Build.cs
@@ -0,0 +1,26 @@
+// Copyright Epic Games, Inc. All Rights Reserved.
+
+using UnrealBuildTool;
+
+public class IREEDriverRDGShaders : ModuleRules
+{
+	public IREEDriverRDGShaders(ReadOnlyTargetRules Target) : base(Target)
+	{
+		PCHUsage = PCHUsageMode.UseExplicitOrSharedPCHs;
+
+		PublicDependencyModuleNames.AddRange(
+			new string[]
+			{
+				"Core",
+				"RenderCore"
+			});
+			
+		PrivateDependencyModuleNames.AddRange(
+			new string[]
+			{
+				"Projects",
+				"Renderer",
+				"RHI"
+			});
+	}
+}
diff --git a/Engine/Plugins/Experimental/NNERuntimeIREE/Source/IREEDriverRDGShaders/Internal/IREEDriverRDGShadersAddCS.h b/Engine/Plugins/Experimental/NNERuntimeIREE/Source/IREEDriverRDGShaders/Internal/IREEDriverRDGShadersAddCS.h
new file mode 100644
--- /dev/null
+++ b/Engine/Plugins/Experimental/NNERuntimeIREE/Source/IREEDriverRDGShaders/Internal/IREEDriverRDGShadersAddCS.h
@@ -0,0 +1,30 @@
+// Copyright Epic Games, Inc. All Rights Reserved.
+
+#pragma once
+
+#include "GlobalShader.h"
+#include "RenderGraphUtils.h"
+#include "ShaderParameterUtils.h"
+
+namespace UE::IREEDriverRDGShaders::Internal
+{
+
+class IREEDRIVERRDGSHADERS_API FAddCS : public FGlobalShader
+{
+    DECLARE_GLOBAL_SHADER(FAddCS);
+    SHADER_USE_PARAMETER_STRUCT(FAddCS, FGlobalShader)
+
+public:
+
+    BEGIN_SHADER_PARAMETER_STRUCT(FParameters, )
+        SHADER_PARAMETER(int32, Size)
+        SHADER_PARAMETER_RDG_BUFFER_SRV(Buffer, A)
+        SHADER_PARAMETER_RDG_BUFFER_SRV(Buffer, B)
+        SHADER_PARAMETER_RDG_BUFFER_UAV(RWBuffer, Result)
+    END_SHADER_PARAMETER_STRUCT()
+
+    static void ModifyCompilationEnvironment(const FGlobalShaderPermutationParameters& InParameters, FShaderCompilerEnvironment& OutEnvironment);
+    static bool ShouldCompilePermutation(const FGlobalShaderPermutationParameters& Parameters);
+};
+
+} // namespace UE::IREEDriverRDGShaders::Private
\ No newline at end of file
diff --git a/Engine/Plugins/Experimental/NNERuntimeIREE/Source/IREEDriverRDGShaders/Internal/IREEDriverRDGShadersAddCrossCS.h b/Engine/Plugins/Experimental/NNERuntimeIREE/Source/IREEDriverRDGShaders/Internal/IREEDriverRDGShadersAddCrossCS.h
new file mode 100644
--- /dev/null
+++ b/Engine/Plugins/Experimental/NNERuntimeIREE/Source/IREEDriverRDGShaders/Internal/IREEDriverRDGShadersAddCrossCS.h
@@ -0,0 +1,29 @@
+// Copyright Epic Games, Inc. All Rights Reserved.
+
+#pragma once
+
+#include "GlobalShader.h"
+#include "RenderGraphUtils.h"
+#include "ShaderParameterUtils.h"
+
+namespace UE::IREEDriverRDGShaders::Internal
+{
+
+class IREEDRIVERRDGSHADERS_API FAddCrossCS : public FGlobalShader
+{
+    DECLARE_GLOBAL_SHADER(FAddCrossCS);
+    SHADER_USE_PARAMETER_STRUCT(FAddCrossCS, FGlobalShader)
+
+public:
+
+    BEGIN_SHADER_PARAMETER_STRUCT(FParameters, )
+        SHADER_PARAMETER_RDG_BUFFER_UAV(RWByteAddressBuffer, _resource_var_0_0_)
+        SHADER_PARAMETER_RDG_BUFFER_UAV(RWByteAddressBuffer, _resource_var_0_1_)
+        SHADER_PARAMETER_RDG_BUFFER_UAV(RWByteAddressBuffer, _resource_var_0_2_)
+    END_SHADER_PARAMETER_STRUCT()
+
+    static void ModifyCompilationEnvironment(const FGlobalShaderPermutationParameters& InParameters, FShaderCompilerEnvironment& OutEnvironment);
+    static bool ShouldCompilePermutation(const FGlobalShaderPermutationParameters& Parameters);
+};
+
+} // namespace UE::IREEDriverRDGShaders::Private
\ No newline at end of file
diff --git a/Engine/Plugins/Experimental/NNERuntimeIREE/Source/IREEDriverRDGShaders/Private/IREEDriverRDGShadersAddCS.cpp b/Engine/Plugins/Experimental/NNERuntimeIREE/Source/IREEDriverRDGShaders/Private/IREEDriverRDGShadersAddCS.cpp
new file mode 100644
--- /dev/null
+++ b/Engine/Plugins/Experimental/NNERuntimeIREE/Source/IREEDriverRDGShaders/Private/IREEDriverRDGShadersAddCS.cpp
@@ -0,0 +1,20 @@
+// Copyright Epic Games, Inc. All Rights Reserved.
+
+#include "IREEDriverRDGShadersAddCS.h"
+
+namespace UE::IREEDriverRDGShaders::Internal
+{
+
+	void FAddCS::ModifyCompilationEnvironment(const FGlobalShaderPermutationParameters& InParameters, FShaderCompilerEnvironment& OutEnvironment)
+	{
+		FGlobalShader::ModifyCompilationEnvironment(InParameters, OutEnvironment);
+	}
+
+	bool FAddCS::ShouldCompilePermutation(const FGlobalShaderPermutationParameters& Parameters)
+	{
+		return true;
+	}
+
+	IMPLEMENT_GLOBAL_SHADER(FAddCS, "/IREEDriverRDGShaders/NNERuntimeIREEAdd.usf", "Add", SF_Compute);
+
+} // UE::NNEDenoiser::Private
\ No newline at end of file
diff --git a/Engine/Plugins/Experimental/NNERuntimeIREE/Source/IREEDriverRDGShaders/Private/IREEDriverRDGShadersAddCrossCS.cpp b/Engine/Plugins/Experimental/NNERuntimeIREE/Source/IREEDriverRDGShaders/Private/IREEDriverRDGShadersAddCrossCS.cpp
new file mode 100644
--- /dev/null
+++ b/Engine/Plugins/Experimental/NNERuntimeIREE/Source/IREEDriverRDGShaders/Private/IREEDriverRDGShadersAddCrossCS.cpp
@@ -0,0 +1,20 @@
+// Copyright Epic Games, Inc. All Rights Reserved.
+
+#include "IREEDriverRDGShadersAddCrossCS.h"
+
+namespace UE::IREEDriverRDGShaders::Internal
+{
+
+	void FAddCrossCS::ModifyCompilationEnvironment(const FGlobalShaderPermutationParameters& InParameters, FShaderCompilerEnvironment& OutEnvironment)
+	{
+		FGlobalShader::ModifyCompilationEnvironment(InParameters, OutEnvironment);
+	}
+
+	bool FAddCrossCS::ShouldCompilePermutation(const FGlobalShaderPermutationParameters& Parameters)
+	{
+		return true;
+	}
+
+	IMPLEMENT_GLOBAL_SHADER(FAddCrossCS, "/IREEDriverRDGShaders/NNERuntimeIREEAdd_Cross.usf", "main", SF_Compute);
+
+} // UE::NNEDenoiser::Private
\ No newline at end of file
diff --git a/Engine/Plugins/Experimental/NNERuntimeIREE/Source/IREEDriverRDGShaders/Private/IREEDriverRDGShadersLog.h b/Engine/Plugins/Experimental/NNERuntimeIREE/Source/IREEDriverRDGShaders/Private/IREEDriverRDGShadersLog.h
new file mode 100644
--- /dev/null
+++ b/Engine/Plugins/Experimental/NNERuntimeIREE/Source/IREEDriverRDGShaders/Private/IREEDriverRDGShadersLog.h
@@ -0,0 +1,8 @@
+// Copyright Epic Games, Inc. All Rights Reserved.
+
+#pragma once
+
+#include "CoreMinimal.h"
+#include "Logging/LogMacros.h"
+
+DECLARE_LOG_CATEGORY_EXTERN(LogIREEDriverRDGShaders, Log, All);
\ No newline at end of file
diff --git a/Engine/Plugins/Experimental/NNERuntimeIREE/Source/IREEDriverRDGShaders/Private/IREEDriverRDGShadersModule.cpp b/Engine/Plugins/Experimental/NNERuntimeIREE/Source/IREEDriverRDGShaders/Private/IREEDriverRDGShadersModule.cpp
new file mode 100644
--- /dev/null
+++ b/Engine/Plugins/Experimental/NNERuntimeIREE/Source/IREEDriverRDGShaders/Private/IREEDriverRDGShadersModule.cpp
@@ -0,0 +1,37 @@
+// Copyright Epic Games, Inc. All Rights Reserved.
+
+#include "CoreMinimal.h"
+#include "Interfaces/IPluginManager.h"
+#include "Misc/Paths.h"
+#include "Modules/ModuleManager.h"
+#include "ShaderCore.h"
+#include "IREEDriverRDGShadersLog.h"
+
+DEFINE_LOG_CATEGORY(LogIREEDriverRDGShaders);
+
+class FIREEDriverRDGShadersModule : public IModuleInterface
+{
+public:
+	// Begin IModuleInterface
+	virtual void StartupModule() override
+	{
+		const TSharedPtr<IPlugin> Plugin = IPluginManager::Get().FindPlugin(TEXT("NNERuntimeIREE"));
+		if (Plugin.IsValid())
+		{
+			const FString ShadersDir = FPaths::Combine(Plugin->GetBaseDir(), TEXT("Shaders"));
+
+			AddShaderSourceDirectoryMapping(TEXT("/IREEDriverRDGShaders"), ShadersDir);
+		}
+		else
+		{
+			UE_LOG(LogIREEDriverRDGShaders, Warning, TEXT("Shaders directory not added. Failed to find IREEDriverRDG plugin"));
+		}
+	}
+
+	virtual void ShutdownModule() override
+	{
+
+	}
+};
+
+IMPLEMENT_MODULE(FIREEDriverRDGShadersModule, IREEDriverRDGShaders);
diff --git a/Engine/Plugins/Experimental/NNERuntimeIREE/Source/NNERuntimeIREE/NNERuntimeIREE.Build.cs b/Engine/Plugins/Experimental/NNERuntimeIREE/Source/NNERuntimeIREE/NNERuntimeIREE.Build.cs
--- a/Engine/Plugins/Experimental/NNERuntimeIREE/Source/NNERuntimeIREE/NNERuntimeIREE.Build.cs
+++ b/Engine/Plugins/Experimental/NNERuntimeIREE/Source/NNERuntimeIREE/NNERuntimeIREE.Build.cs
@@ -1,5 +1,6 @@
 // Copyright Epic Games, Inc. All Rights Reserved.
 
+using System.IO;
 using UnrealBuildTool;
 
 public class NNERuntimeIREE : ModuleRules
@@ -16,8 +17,11 @@ public class NNERuntimeIREE : ModuleRules
 				"CoreUObject",
 				"Engine",
 				"IREE",
+				"IREEDriverRDG",
 				"NNE",
-				"Projects"
+				"Projects",
+				"RenderCore",
+				"RHI"
 			}
 		);
 
@@ -47,5 +51,8 @@ public class NNERuntimeIREE : ModuleRules
 			PublicDefinitions.Add("WITH_NNE_RUNTIME_IREE");
 			PrivateDefinitions.Add("NNE_RUNTIME_IREE_USE_COMBINED_LIB_PATH");
 		}
+
+		// Disable NNERuntimeIREERDG until development is progressed and runtime is functional
+		PublicDefinitions.Add("WITH_NNE_RUNTIME_IREE_RDG=false");
 	}
 }
diff --git a/Engine/Plugins/Experimental/NNERuntimeIREE/Source/NNERuntimeIREE/Private/NNERuntimeIREE.cpp b/Engine/Plugins/Experimental/NNERuntimeIREE/Source/NNERuntimeIREE/Private/NNERuntimeIREE.cpp
--- a/Engine/Plugins/Experimental/NNERuntimeIREE/Source/NNERuntimeIREE/Private/NNERuntimeIREE.cpp
+++ b/Engine/Plugins/Experimental/NNERuntimeIREE/Source/NNERuntimeIREE/Private/NNERuntimeIREE.cpp
@@ -50,32 +50,47 @@ namespace UE::NNERuntimeIREE::CPU::Private
 		}
 	}
 
-	FString GetModelDataIdentifier(const FString& RuntimeName, const FGuid& Guid, const FString& FileIdString, const FString& PlatformName, const FString& Architecture)
+	FString GetModelDataIdentifier(const FString& RuntimeName, const FGuid& Guid, int32 Version, const FString& FileIdString, const FString& PlatformName, const FString& Architecture)
 	{
-		return RuntimeName + "-" + Guid.ToString(EGuidFormats::Digits) + "-" + FString::FromInt(UNNERuntimeIREECpu::Version) + "-" + FileIdString + "-" + PlatformName + (!Architecture.IsEmpty() ? ("-" + Architecture) : "");
+		return RuntimeName + "-" + Guid.ToString(EGuidFormats::Digits) + "-" + FString::FromInt(Version) + "-" + FileIdString + "-" + PlatformName + (!Architecture.IsEmpty() ? ("-" + Architecture) : "");
 	}
 
-	FString GetIntermediateModelDirPath(const FString& PlatformName, const FString& ModelName)
+	FString GuidToString(const FGuid& Guid)
 	{
-		return FPaths::Combine("Intermediate", "Build", GetBinariesSubdirectory(PlatformName), UE_PLUGIN_NAME, ModelName);
+		return Guid.ToString(EGuidFormats::Digits).ToLower();
 	}
 
-	FString GetStagedModelDirPath(const FString& PlatformName)
+	FString GetRuntimeIdAndVersionString(const FGuid& RuntimeGuid, int32 Version)
 	{
-		return FPaths::Combine("Binaries", GetBinariesSubdirectory(PlatformName), UE_PLUGIN_NAME);
+		return GuidToString(RuntimeGuid) + "-" + FString::FromInt(Version);
 	}
 
-	FString GetPackagedModelDirPath(const FString& PlatformName)
+	FString GetModelDataDirectoryName(const FString& RuntimeIdAndVersionString, const FString& FileIdString)
 	{
-		return GetStagedModelDirPath(PlatformName);
+		return RuntimeIdAndVersionString + "_" + FileIdString;
 	}
 
-	FString GetSharedLibDirPath(const FString& PlatformName, const FString& ModelName)
+	FString GetIntermediateModelDirPath(const FString& PlatformName, const FString& RuntimeName, const FString& ModelName)
+	{
+		return FPaths::Combine("Intermediate", "Build", GetBinariesSubdirectory(PlatformName), UE_PLUGIN_NAME, RuntimeName, ModelName);
+	}
+
+	FString GetStagedModelDirPath(const FString& PlatformName, const FString& RuntimeName)
+	{
+		return FPaths::Combine("Binaries", GetBinariesSubdirectory(PlatformName), UE_PLUGIN_NAME, RuntimeName);
+	}
+
+	FString GetPackagedModelDirPath(const FString& PlatformName, const FString& RuntimeName)
+	{
+		return GetStagedModelDirPath(PlatformName, RuntimeName);
+	}
+
+	FString GetSharedLibDirPath(const FString& PlatformName, const FString& RuntimeName, const FString& ModelName)
 	{
 #if WITH_EDITOR
-	return GetIntermediateModelDirPath(PlatformName, ModelName);
+	return GetIntermediateModelDirPath(PlatformName, RuntimeName, ModelName);
 #else
-	return GetPackagedModelDirPath(PlatformName);
+	return GetPackagedModelDirPath(PlatformName, RuntimeName);
 #endif // WITH_EDITOR
 	}
 } // UE::NNERuntimeIREE::CPU::Private
@@ -120,8 +135,8 @@ TSharedPtr<UE::NNE::FSharedModelData> UNNERuntimeIREECpu::CreateModelData(const
 	IPlatformFile& PlatformFile = FPlatformFileManager::Get().GetPlatformFile();
 
 	const FString FileIdString = FileId.ToString(EGuidFormats::Digits).ToLower();
-	const FString IntermediateDirFullPath = FPaths::ConvertRelativePathToFull(FPaths::Combine(FPaths::ProjectDir(), GetIntermediateModelDirPath(TargetPlatformName, FileIdString)));
-	const FString SharedLibraryDirFullPath = FPaths::ConvertRelativePathToFull(FPaths::Combine(FPaths::ProjectDir(), GetSharedLibDirPath(TargetPlatformName, FileIdString)));
+	const FString IntermediateDirFullPath = FPaths::ConvertRelativePathToFull(FPaths::Combine(FPaths::ProjectDir(), GetIntermediateModelDirPath(TargetPlatformName, GetRuntimeName(), FileIdString)));
+	const FString SharedLibraryDirFullPath = FPaths::ConvertRelativePathToFull(FPaths::Combine(FPaths::ProjectDir(), GetSharedLibDirPath(TargetPlatformName, GetRuntimeName(), FileIdString)));
 
 	FString IREEModelDataFilePath = FPaths::Combine(IntermediateDirFullPath, FileIdString) + ".ireemodeldata";
 
@@ -210,7 +225,7 @@ TSharedPtr<UE::NNE::FSharedModelData> UNNERuntimeIREECpu::CreateModelData(const
 	}
 
 	// Copy files for staging
-	FString StagingDirFullPath = FPaths::ConvertRelativePathToFull(FPaths::Combine(FPaths::ProjectDir(), GetPackagedModelDirPath(TargetPlatformName)));
+	FString StagingDirFullPath = FPaths::ConvertRelativePathToFull(FPaths::Combine(FPaths::ProjectDir(), GetPackagedModelDirPath(TargetPlatformName, GetRuntimeName())));
 	for (int32 i = 0; i < CompilerResult.ArchitectureInfos.Num(); i++)
 	{
 		SCOPED_NAMED_EVENT_TEXT("Copy", FColor::Magenta);
@@ -239,7 +254,7 @@ FString UNNERuntimeIREECpu::GetModelDataIdentifier(const FString& FileType, TCon
 {
 	// Leave architecture blank as there is only one model data for all architectures of a given platform, only the vmfb and shared lib are different
 	FString PlatformName = UE::NNERuntimeIREE::CPU::Private::GetTargetPlatformName(TargetPlatform);
-	return UE::NNERuntimeIREE::CPU::Private::GetModelDataIdentifier(GetRuntimeName(), UNNERuntimeIREECpu::GUID, FileId.ToString(EGuidFormats::Digits), PlatformName, "");
+	return UE::NNERuntimeIREE::CPU::Private::GetModelDataIdentifier(GetRuntimeName(), UNNERuntimeIREECpu::GUID, UNNERuntimeIREECpu::Version, FileId.ToString(EGuidFormats::Digits), PlatformName, "");
 }
 
 UNNERuntimeIREECpu::ECanCreateModelCPUStatus UNNERuntimeIREECpu::CanCreateModelCPU(const TObjectPtr<UNNEModelData> ModelData) const
@@ -332,7 +347,7 @@ TSharedPtr<UE::NNE::IModelCPU> UNNERuntimeIREECpu::CreateModelCPU(const TObjectP
 	const FNNERuntimeIREEArchitectureInfoCPU& ArchitectureInfo = CompilerResult.ArchitectureInfos[ArchitectureIndex];
 
 	const FString FileIdString = IREEModelData->FileId.ToString(EGuidFormats::Digits).ToLower();
-	const FString SharedLibraryDirFullPath = FPaths::ConvertRelativePathToFull(FPaths::Combine(FPaths::ProjectDir(), GetSharedLibDirPath(UGameplayStatics::GetPlatformName(), FileIdString)));
+	const FString SharedLibraryDirFullPath = FPaths::ConvertRelativePathToFull(FPaths::Combine(FPaths::ProjectDir(), GetSharedLibDirPath(UGameplayStatics::GetPlatformName(), GetRuntimeName(), FileIdString)));
 	const FString SharedLibrarySubDirFullPath = FPaths::Combine(SharedLibraryDirFullPath, ArchitectureInfo.RelativeDirPath);
 
 	TSharedPtr<UE::NNE::IModelCPU> Model = UE::NNERuntimeIREE::CPU::FModel::Make(SharedLibrarySubDirFullPath, ArchitectureInfo.SharedLibraryFileName, ArchitectureInfo.VmfbFileName, ArchitectureInfo.SharedLibraryEntryPointName, *ModuleMetaData);
@@ -373,7 +388,7 @@ TSharedPtr<UE::NNE::FSharedModelData> UNNERuntimeIREEGpu::CreateModelData(const
 FString UNNERuntimeIREEGpu::GetModelDataIdentifier(const FString& FileType, TConstArrayView64<uint8> FileData, const TMap<FString, TConstArrayView64<uint8>>& AdditionalFileData, const FGuid& FileId, const ITargetPlatform* TargetPlatform) const
 {
 	FString PlatformName = UE::NNERuntimeIREE::CPU::Private::GetTargetPlatformName(TargetPlatform);
-	return UE::NNERuntimeIREE::CPU::Private::GetModelDataIdentifier(GetRuntimeName(), GetGUID(), FileId.ToString(EGuidFormats::Digits), PlatformName, "");
+	return UE::NNERuntimeIREE::CPU::Private::GetModelDataIdentifier(GetRuntimeName(), GetGUID(), UNNERuntimeIREEGpu::GetVersion(), FileId.ToString(EGuidFormats::Digits), PlatformName, "");
 }
 
 UNNERuntimeIREEGpu::ECanCreateModelGPUStatus UNNERuntimeIREEGpu::CanCreateModelGPU(const TObjectPtr<UNNEModelData> ModelData) const
@@ -506,13 +521,140 @@ UNNERuntimeIREERdg::ECanCreateModelDataStatus UNNERuntimeIREERdg::CanCreateModel
 
 TSharedPtr<UE::NNE::FSharedModelData> UNNERuntimeIREERdg::CreateModelData(const FString& FileType, TConstArrayView64<uint8> FileData, const TMap<FString, TConstArrayView64<uint8>>& AdditionalFileData, const FGuid& FileId, const ITargetPlatform* TargetPlatform)
 {
+	SCOPED_NAMED_EVENT_TEXT("UNNERuntimeIREERdg::CreateModelData", FColor::Magenta);
+
+#if WITH_EDITOR
+	using namespace UE::NNERuntimeIREE::CPU::Private;
+
+	const FString TargetPlatformName = GetTargetPlatformName(TargetPlatform);
+	if (CanCreateModelData(FileType, FileData, AdditionalFileData, FileId, TargetPlatform) != ECanCreateModelDataStatus::Ok)
+	{
+		UE_LOG(LogNNERuntimeIREE, Warning, TEXT("UNNERuntimeIREERdg cannot create the model data with id %s (Filetype: %s) for platform %s"), *FileId.ToString(EGuidFormats::Digits).ToLower(), *FileType, *TargetPlatformName);
+		return TSharedPtr<UE::NNE::FSharedModelData>();
+	}
+
+	IPlatformFile& PlatformFile = FPlatformFileManager::Get().GetPlatformFile();
+
+	const FString FileIdString = FileId.ToString(EGuidFormats::Digits).ToLower();
+	const FString IntermediateDirFullPath = FPaths::ConvertRelativePathToFull(FPaths::Combine(FPaths::ProjectDir(), GetIntermediateModelDirPath(TargetPlatformName, GetRuntimeName(), FileIdString)));
+	const FString SharedLibraryDirFullPath = FPaths::ConvertRelativePathToFull(FPaths::Combine(FPaths::ProjectDir(), GetSharedLibDirPath(TargetPlatformName, GetRuntimeName(), FileIdString)));
+
+	FString IREEModelDataFilePath = FPaths::Combine(IntermediateDirFullPath, FileIdString) + ".ireemodeldata";
+
+	TArray64<uint8> ResultData;
+	TWeakObjectPtr<UNNERuntimeIREEModelData> IREEModelData = NewObject<UNNERuntimeIREEModelData>();
+	FNNERuntimeIREECompilerResultCPU CompilerResult{};
+
+	bool bNeedCompileMlir = true;
+	if (PlatformFile.FileExists(*IREEModelDataFilePath))
+	{
+		SCOPED_NAMED_EVENT_TEXT("Validate", FColor::Magenta);
+
+		FFileHelper::LoadFileToArray(ResultData, *IREEModelDataFilePath);
+		
+		{
+			FMemoryReaderView Reader(ResultData, /*bIsPersitent =*/ true);
+			IREEModelData->Serialize(Reader);
+		}
+
+		check(FileIdString.Equals(IREEModelData->FileId.ToString(EGuidFormats::Digits).ToLower()));
+
+		{
+			FMemoryReaderView Reader(IREEModelData->CompilerResult, /*bIsPersitent =*/ true);
+			FNNERuntimeIREECompilerResultCPU::StaticStruct()->SerializeBin(Reader, &CompilerResult);
+		}
+
+		bNeedCompileMlir = false;
+		for (int32 i = 0; i < CompilerResult.ArchitectureInfos.Num(); i++)
+		{
+			const FNNERuntimeIREEArchitectureInfoCPU& Info = CompilerResult.ArchitectureInfos[i];
+			const FString SharedLibrarySubDirFullPath = FPaths::Combine(SharedLibraryDirFullPath, Info.RelativeDirPath);
+
+			// const FString SharedLibraryFilePath = FPaths::Combine(SharedLibrarySubDirFullPath, Info.SharedLibraryFileName);
+			const FString VmfbFilePath = FPaths::Combine(SharedLibrarySubDirFullPath, Info.VmfbFileName);
+
+			// bNeedCompileMlir |= !PlatformFile.FileExists(*SharedLibraryFilePath);
+			bNeedCompileMlir |= !PlatformFile.FileExists(*VmfbFilePath);
+		}
+	}
+	
+	if (bNeedCompileMlir)
+	{
+		SCOPED_NAMED_EVENT_TEXT("Compile", FColor::Magenta);
+
+		PlatformFile.DeleteDirectoryRecursively(*IntermediateDirFullPath);
+		PlatformFile.CreateDirectoryTree(*IntermediateDirFullPath);
+
+		TUniquePtr<UE::NNERuntimeIREE::CPU::FCompiler> Compiler = UE::NNERuntimeIREE::CPU::FCompiler::Make(TargetPlatformName);
+		if (!Compiler.IsValid())
+		{
+			UE_LOG(LogNNERuntimeIREE, Warning, TEXT("UNNERuntimeIREERdg failed to create a compiler to compile for platform %s"), *TargetPlatformName);
+			return TSharedPtr<UE::NNE::FSharedModelData>();
+		}
+
+		TWeakObjectPtr<UNNERuntimeIREEModuleMetaData> CompilerModuleMetaData = NewObject<UNNERuntimeIREEModuleMetaData>();
+
+		if (!Compiler->CompileMlirRDG(FileData, FileIdString, IntermediateDirFullPath, CompilerResult, *CompilerModuleMetaData))
+		{
+			UE_LOG(LogNNERuntimeIREE, Warning, TEXT("UNNERuntimeIREERdg failed to compile model %s"), *FileIdString);
+			return TSharedPtr<UE::NNE::FSharedModelData>();
+		}
+
+		IREEModelData->GUID = UNNERuntimeIREERdg::GUID;
+		IREEModelData->Version = UNNERuntimeIREERdg::Version;
+		IREEModelData->FileId = FileId;
+		if (AdditionalFileData.Contains("IREEModuleMetaData"))
+		{
+			IREEModelData->ModuleMetaData = AdditionalFileData["IREEModuleMetaData"];
+		}
+		if (IREEModelData->ModuleMetaData.IsEmpty())
+		{
+			FMemoryWriter64 Writer(IREEModelData->ModuleMetaData, /*bIsPersitent =*/ true);
+			CompilerModuleMetaData->Serialize(Writer);
+		}
+		{
+			FMemoryWriter64 Writer(IREEModelData->CompilerResult, /*bIsPersitent =*/ true);
+			FNNERuntimeIREECompilerResultCPU::StaticStruct()->SerializeBin(Writer, &CompilerResult);
+		}
+
+		{
+			FMemoryWriter64 Writer(ResultData, /*bIsPersitent =*/ true);
+			IREEModelData->Serialize(Writer);
+		}
+
+		FFileHelper::SaveArrayToFile(ResultData, *IREEModelDataFilePath);
+	}
+
+	// Copy files for staging
+	FString StagingDirFullPath = FPaths::ConvertRelativePathToFull(FPaths::Combine(FPaths::ProjectDir(), GetPackagedModelDirPath(TargetPlatformName, GetRuntimeName())));
+	for (int32 i = 0; i < CompilerResult.ArchitectureInfos.Num(); i++)
+	{
+		SCOPED_NAMED_EVENT_TEXT("Copy", FColor::Magenta);
+
+		const FNNERuntimeIREEArchitectureInfoCPU& Info = CompilerResult.ArchitectureInfos[i];
+		const FString SharedLibrarySubDirFullPath = FPaths::Combine(SharedLibraryDirFullPath, Info.RelativeDirPath);
+		const FString StagingSubDirFullPath = FPaths::Combine(StagingDirFullPath, Info.Architecture);
+		
+		// const FString SharedLibraryFilePathSrc = FPaths::Combine(SharedLibrarySubDirFullPath, Info.SharedLibraryFileName);
+		const FString VmfbFilePathSrc = FPaths::Combine(SharedLibrarySubDirFullPath, Info.VmfbFileName);
+
+		// const FString SharedLibraryFilePathDest = FPaths::Combine(StagingSubDirFullPath, Info.SharedLibraryFileName);
+		const FString VmfbFilePathDest = FPaths::Combine(StagingSubDirFullPath, Info.VmfbFileName);
+
+		// IFileManager::Get().Copy(*SharedLibraryFilePathDest, *SharedLibraryFilePathSrc, bNeedCompileMlir);
+		IFileManager::Get().Copy(*VmfbFilePathDest, *VmfbFilePathSrc, bNeedCompileMlir);
+	}
+
+	return MakeShared<UE::NNE::FSharedModelData>(MakeSharedBufferFromArray(MoveTemp(ResultData)), 0);
+#else
 	return TSharedPtr<UE::NNE::FSharedModelData>();
+#endif // WITH_EDITOR
 }
 
 FString UNNERuntimeIREERdg::GetModelDataIdentifier(const FString& FileType, TConstArrayView64<uint8> FileData, const TMap<FString, TConstArrayView64<uint8>>& AdditionalFileData, const FGuid& FileId, const ITargetPlatform* TargetPlatform) const
 {
 	FString PlatformName = UE::NNERuntimeIREE::CPU::Private::GetTargetPlatformName(TargetPlatform);
-	return UE::NNERuntimeIREE::CPU::Private::GetModelDataIdentifier(GetRuntimeName(), UNNERuntimeIREERdg::GUID, FileId.ToString(EGuidFormats::Digits), PlatformName, "");
+	return UE::NNERuntimeIREE::CPU::Private::GetModelDataIdentifier(GetRuntimeName(), UNNERuntimeIREERdg::GUID, UNNERuntimeIREERdg::Version, FileId.ToString(EGuidFormats::Digits), PlatformName, "");
 }
 
 UNNERuntimeIREERdg::ECanCreateModelRDGStatus UNNERuntimeIREERdg::CanCreateModelRDG(const TObjectPtr<UNNEModelData> ModelData) const
@@ -543,23 +685,90 @@ TSharedPtr<UE::NNE::IModelRDG> UNNERuntimeIREERdg::CreateModelRDG(const TObjectP
 {
 	check(ModelData != nullptr);
 
+	using namespace UE::NNERuntimeIREE::CPU::Private;
+
 	if (CanCreateModelRDG(ModelData) != ECanCreateModelRDGStatus::Ok)
 	{
 		UE_LOG(LogNNERuntimeIREE, Warning, TEXT("UNNERuntimeIREERdg cannot create a model from the model data with id %s"), *ModelData->GetFileId().ToString(EGuidFormats::Digits));
 		return TSharedPtr<UE::NNE::IModelRDG>();
 	}
 
-	check(ModelData->GetModelData(GetRuntimeName()).IsValid());
+	FString CurrentArchitecture = "";
+#if PLATFORM_CPU_X86_FAMILY
+	CurrentArchitecture = "x86_64";
+#elif PLATFORM_CPU_ARM_FAMILY
+	CurrentArchitecture = "arm64";
+#endif
 
-	UE::NNE::IModelRDG* IModel = nullptr;
-	TConstArrayView<uint8> SharedDataView = ModelData->GetModelData(GetRuntimeName())->GetView();
+	TSharedPtr<UE::NNE::FSharedModelData> SharedData = ModelData->GetModelData(GetRuntimeName());
+	check(SharedData.IsValid());
+
+	TConstArrayView<uint8> SharedDataView = SharedData->GetView();
+
+	TWeakObjectPtr<UNNERuntimeIREEModelData> IREEModelData = NewObject<UNNERuntimeIREEModelData>();
+	{
+		FMemoryReaderView Reader(SharedDataView, /*bIsPersitent =*/ true);
+		IREEModelData->Serialize(Reader);
+	}
+
+	if (IREEModelData->ModuleMetaData.IsEmpty())
+	{
+		UE_LOG(LogNNERuntimeIREE, Warning, TEXT("UNNERuntimeIREERdg failed to find any module meta data, please reimport the original model"));
+		return {};
+	}
+
+	TWeakObjectPtr<UNNERuntimeIREEModuleMetaData> ModuleMetaData = NewObject<UNNERuntimeIREEModuleMetaData>();
+	{
+		FMemoryReaderView Reader(IREEModelData->ModuleMetaData, /*bIsPersitent =*/ true);
+		ModuleMetaData->Serialize(Reader);
+	}
+
+	if (ModuleMetaData->FunctionMetaData.IsEmpty())
+	{
+		UE_LOG(LogNNERuntimeIREE, Warning, TEXT("UNNERuntimeIREERdg failed to parse the module meta data, please reimport the original model"));
+		return {};
+	}
+
+	FNNERuntimeIREECompilerResultCPU CompilerResult{};
+	{
+		FMemoryReaderView Reader(IREEModelData->CompilerResult, /*bIsPersitent =*/ true);
+		FNNERuntimeIREECompilerResultCPU::StaticStruct()->SerializeBin(Reader, &CompilerResult);
+	}
+
+	int32 ArchitectureIndex = -1;
+	for (int32 i = 0; i < CompilerResult.ArchitectureInfos.Num(); i++)
+	{
+		if ((CompilerResult.ArchitectureInfos[i].Architecture.IsEmpty() && ArchitectureIndex < 0) || CompilerResult.ArchitectureInfos[i].Architecture.Equals(CurrentArchitecture))
+		{
+			ArchitectureIndex = i;
+		}
+	}
+
+	if (ArchitectureIndex < 0)
+	{
+		UE_LOG(LogNNERuntimeIREE, Warning, TEXT("UNNERuntimeIREERdg failed to find a matching architecture for \'%s\'"), *CurrentArchitecture);
+		return {};
+	}
 
-	return TSharedPtr<UE::NNE::IModelRDG>(IModel);
+	const FNNERuntimeIREEArchitectureInfoCPU& ArchitectureInfo = CompilerResult.ArchitectureInfos[ArchitectureIndex];
+
+	const FString FileIdString = IREEModelData->FileId.ToString(EGuidFormats::Digits).ToLower();
+	const FString SharedLibraryDirFullPath = FPaths::ConvertRelativePathToFull(FPaths::Combine(FPaths::ProjectDir(), GetSharedLibDirPath(UGameplayStatics::GetPlatformName(), GetRuntimeName(), FileIdString)));
+	const FString SharedLibrarySubDirFullPath = FPaths::Combine(SharedLibraryDirFullPath, ArchitectureInfo.RelativeDirPath);
+
+	TSharedPtr<UE::NNE::IModelRDG> Model = UE::NNERuntimeIREE::RDG::FModel::Make(SharedLibrarySubDirFullPath, ArchitectureInfo.VmfbFileName, *ModuleMetaData);
+	if (!Model.IsValid())
+	{
+		UE_LOG(LogNNERuntimeIREE, Warning, TEXT("UNNERuntimeIREERdg could not initialize the model created from model data with id %s"), *FileIdString);
+		return {};
+	}
+
+	return Model;
 }
 
 bool UNNERuntimeIREERdg::IsAvailable() const
 {
-	return false;
+	return WITH_NNE_RUNTIME_IREE_RDG;
 }
 
 #endif // WITH_NNE_RUNTIME_IREE
\ No newline at end of file
diff --git a/Engine/Plugins/Experimental/NNERuntimeIREE/Source/NNERuntimeIREE/Private/NNERuntimeIREECompiler.cpp b/Engine/Plugins/Experimental/NNERuntimeIREE/Source/NNERuntimeIREE/Private/NNERuntimeIREECompiler.cpp
--- a/Engine/Plugins/Experimental/NNERuntimeIREE/Source/NNERuntimeIREE/Private/NNERuntimeIREECompiler.cpp
+++ b/Engine/Plugins/Experimental/NNERuntimeIREE/Source/NNERuntimeIREE/Private/NNERuntimeIREECompiler.cpp
@@ -352,6 +352,67 @@ namespace UE::NNERuntimeIREE
 
 			return bResult;
 		}
+		
+		bool FCompiler::CompileMlirRDG(TConstArrayView<uint8> InFileData, const FString& InModelName, const FString& InOutputDir, FNNERuntimeIREECompilerResultCPU& OutCompilerResult, UNNERuntimeIREEModuleMetaData& ModuleMetaData)
+		{
+			using namespace Private;
+
+			{
+				SCOPED_NAMED_EVENT_TEXT("Metadata", FColor::Magenta);
+
+				FString FileDataString = "";
+				FileDataString.AppendChars((char*)InFileData.GetData(), InFileData.Num());
+				ModuleMetaData.ParseFromString(FileDataString);
+			}
+
+			IPlatformFile& PlatformFile = FPlatformFileManager::Get().GetPlatformFile();
+
+			FString InputFilePath = FPaths::Combine(InOutputDir, InModelName) + ".mlir";
+			if (!PlatformFile.FileExists(*InputFilePath))
+			{
+				SCOPED_NAMED_EVENT_TEXT("InputFile", FColor::Magenta);
+
+				FFileHelper::SaveArrayToFile(InFileData, *InputFilePath);
+			}
+
+			bool bResult = true;
+			for (int32 i = 0; i < BuildTargets.Num(); i++)
+			{
+				FString IntermediateDirPath = FPaths::Combine(InOutputDir, BuildTargets[i].Architecture);
+				PlatformFile.CreateDirectoryTree(*IntermediateDirPath);
+				FString IntermediateFilePathNoExt = FPaths::Combine(IntermediateDirPath, InModelName);
+				FString VmfbFilePath = IntermediateFilePathNoExt + ".vmfb";
+
+				FString CompilerArguments = TEXT("--iree-hal-target-backends=vulkan-spirv -o ${VMFB_PATH} ${INPUT_PATH}");
+				CompilerArguments.ReplaceInline(*FString("${VMFB_PATH}"), *(FString("\"") + VmfbFilePath + "\""));
+				CompilerArguments.ReplaceInline(*FString("${INPUT_PATH}"), *(FString("\"") + InputFilePath + "\""));
+
+				RunCommand(CompilerCommand, CompilerArguments, IntermediateFilePathNoExt + "_compile-log.txt");
+
+				if (!PlatformFile.FileExists(*VmfbFilePath))
+				{
+					UE_LOG(LogNNERuntimeIREE, Warning, TEXT("UNNERuntimeIREERdg failed to compile the model \"%s\" using the command:"), *InputFilePath);
+					UE_LOG(LogNNERuntimeIREE, Warning, TEXT("\"%s\" %s"), *CompilerCommand, *CompilerArguments);
+					bResult = false;
+					continue;
+				}
+
+				FNNERuntimeIREEArchitectureInfoCPU ArchitectureInfo;
+				ArchitectureInfo.Architecture = BuildTargets[i].Architecture;
+				ArchitectureInfo.RelativeDirPath = BuildTargets[i].Architecture;
+				ArchitectureInfo.VmfbFileName = InModelName + ".vmfb";
+				OutCompilerResult.ArchitectureInfos.Add(MoveTemp(ArchitectureInfo));
+			}
+
+			bResult &= !OutCompilerResult.ArchitectureInfos.IsEmpty();
+			
+			if (!bResult)
+			{
+				OutCompilerResult.ArchitectureInfos.Empty();
+			}
+			
+			return bResult;
+		}
 	} // CPU
 } // UE::NNERuntimeIREE
 
diff --git a/Engine/Plugins/Experimental/NNERuntimeIREE/Source/NNERuntimeIREE/Private/NNERuntimeIREECompiler.h b/Engine/Plugins/Experimental/NNERuntimeIREE/Source/NNERuntimeIREE/Private/NNERuntimeIREECompiler.h
--- a/Engine/Plugins/Experimental/NNERuntimeIREE/Source/NNERuntimeIREE/Private/NNERuntimeIREECompiler.h
+++ b/Engine/Plugins/Experimental/NNERuntimeIREE/Source/NNERuntimeIREE/Private/NNERuntimeIREECompiler.h
@@ -82,6 +82,7 @@ namespace UE::NNERuntimeIREE::CPU
 
 	public:
 		bool CompileMlir(TConstArrayView<uint8> InFileData, const FString& InModelName, const FString& InOutputDir, FNNERuntimeIREECompilerResultCPU& OutCompilerResult, UNNERuntimeIREEModuleMetaData& ModuleMetaData);
+		bool CompileMlirRDG(TConstArrayView<uint8> InFileData, const FString& InModelName, const FString& InOutputDir, FNNERuntimeIREECompilerResultCPU& OutCompilerResult, UNNERuntimeIREEModuleMetaData& ModuleMetaData);
 
 	private:
 		FString CompilerCommand;
diff --git a/Engine/Plugins/Experimental/NNERuntimeIREE/Source/NNERuntimeIREE/Private/NNERuntimeIREEModel.cpp b/Engine/Plugins/Experimental/NNERuntimeIREE/Source/NNERuntimeIREE/Private/NNERuntimeIREEModel.cpp
--- a/Engine/Plugins/Experimental/NNERuntimeIREE/Source/NNERuntimeIREE/Private/NNERuntimeIREEModel.cpp
+++ b/Engine/Plugins/Experimental/NNERuntimeIREE/Source/NNERuntimeIREE/Private/NNERuntimeIREEModel.cpp
@@ -6,6 +6,8 @@
 
 #include "GenericPlatform/GenericPlatformProcess.h"
 #include "HAL/FileManager.h"
+#include "HAL/Platform.h"
+#include "IREEDriverRDG.h"
 #include "Memory/SharedBuffer.h"
 #include "Misc/Paths.h"
 #include "Misc/ScopeLock.h"
@@ -13,6 +15,8 @@
 #include "NNEModelData.h"
 #include "NNERuntimeIREELog.h"
 #include "NNEStatus.h"
+#include "NNETensor.h"
+#include "RenderGraphUtils.h"
 #include "Serialization/Archive.h"
 
 #if PLATFORM_MICROSOFT
@@ -32,6 +36,13 @@ THIRD_PARTY_INCLUDES_END
 #include "Microsoft/HideMicrosoftPlatformTypes.h"
 #endif // PLATFORM_MICROSOFT
 
+BEGIN_SHADER_PARAMETER_STRUCT(FIREESessionRDGParameters, )
+	RDG_BUFFER_ACCESS_ARRAY(InputBuffers)
+	RDG_BUFFER_ACCESS_ARRAY(OutputBuffers)
+END_SHADER_PARAMETER_STRUCT()
+
+DECLARE_GPU_STAT_NAMED(FNNERuntimeIREERDG, TEXT("EnqueueRDG"));
+
 namespace UE::NNERuntimeIREE
 {
 	namespace Private
@@ -416,7 +427,6 @@ namespace UE::NNERuntimeIREE
 			bool GetFunctionByName(const FString& InFunctionName, iree_vm_function_t* OutFunction)
 			{
 				iree_status_t Status = iree_ok_status();
-				check(iree_status_is_ok(Status));
 
 				bool bFound = false;
 				iree_host_size_t Ordinal = 0;
@@ -445,7 +455,6 @@ namespace UE::NNERuntimeIREE
 				}
 
 				*OutFunction = Function;
-				iree_status_free(Status);
 				return true;
 			}
 		};
@@ -992,6 +1001,442 @@ namespace UE::NNERuntimeIREE
 			return FModelInstance::Make(Device, Module);
 		}
 	} // CPU
+
+	namespace RDG
+	{
+		namespace Private
+		{
+			class FDevice
+			{
+				FDevice(TSharedRef<UE::NNERuntimeIREE::Private::FInstance> InInstance, iree_hal_device_t* InDevice) : Instance(InInstance), Device(InDevice)
+				{
+					check(InDevice);
+				}
+
+			public:
+				~FDevice()
+				{
+					iree_hal_device_release(Device);
+				}
+
+				static TSharedPtr<FDevice> Make()
+				{
+					TSharedPtr<UE::NNERuntimeIREE::Private::FInstance> Instance = UE::NNERuntimeIREE::Private::FInstance::GetInstance();
+					if (!Instance.IsValid())
+					{
+						return {};
+					}
+
+					iree_hal_device_t* Device = nullptr;
+					iree_status_t Status = IREE::HAL::RDG::DeviceCreate(iree_allocator_system(), &Device);
+					if (!iree_status_is_ok(Status))
+					{
+						UE_LOG(LogNNERuntimeIREE, Error, TEXT("Could not create IREE RDG device!"));
+						iree_status_free(Status);
+						return {};
+					}
+
+					TSharedPtr<FDevice> Result = TSharedPtr<FDevice>(new FDevice(Instance.ToSharedRef(), Device));
+					
+					Devices.Add("RDG", Result);
+					
+					return Result;
+				}
+
+				bool CreateSession(iree_runtime_session_t** OutSession)
+				{
+					return Instance->CreateSession(Device, OutSession);
+				}
+
+				iree_hal_allocator_t* GetDeviceAllocator() const
+				{
+					return iree_hal_device_allocator(Device);
+				}
+
+			private:
+				TSharedRef<UE::NNERuntimeIREE::Private::FInstance> Instance;
+				iree_hal_device_t* Device;
+				static TMap<FString, TWeakPtr<FDevice>> Devices;
+			};
+			TMap<FString, TWeakPtr<FDevice>> FDevice::Devices;
+
+			class FGraphBuilderSetter
+			{
+			public:
+				FGraphBuilderSetter(TSharedRef<FDevice> Device, FRDGBuilder& GraphBuilder) : Device(Device)
+				{
+					IREE::HAL::RDG::DeviceAllocatorSetGraphBuilder(Device->GetDeviceAllocator(), GraphBuilder);
+				}
+
+				~FGraphBuilderSetter()
+				{
+					IREE::HAL::RDG::DeviceAllocatorResetGraphBuilder(Device->GetDeviceAllocator());
+				}
+
+			private:
+				TSharedRef<FDevice> Device;
+			};
+			class FSession
+			{
+			public:
+				using ESetInputTensorShapesStatus = UE::NNE::IModelInstanceRDG::ESetInputTensorShapesStatus;
+				using EEnqueueRDGStatus = UE::NNE::IModelInstanceRDG::EEnqueueRDGStatus;
+
+				~FSession()
+				{
+					iree_runtime_call_deinitialize(&Call);
+					iree_runtime_session_release(Session);
+				}
+
+				static TSharedPtr<FSession> Make(TSharedRef<FDevice> InDevice, TSharedRef<UE::NNERuntimeIREE::Private::FModule> InModule)
+				{
+					check(!InModule->GetFunctionMetaDataView().IsEmpty());
+
+					iree_runtime_session_t* Session = nullptr;
+					if (!InDevice->CreateSession(&Session))
+					{
+						return {};
+					}
+
+					if (!InModule->AppendToSession(Session))
+					{
+						iree_runtime_session_release(Session);
+						return {};
+					}
+
+					FString MainFunctionName = InModule->GetFunctionMetaDataView()[0].Name;
+					iree_vm_function_t MainFunction;
+					if (!InModule->GetFunctionByName(MainFunctionName, &MainFunction))
+					{
+						iree_runtime_session_release(Session);
+						return {};
+					}
+
+					iree_host_size_t NumInputs = 0;
+					iree_host_size_t NumOutputs = 0;
+					iree_vm_function_signature_t Signature = iree_vm_function_signature(&MainFunction);
+					iree_status_t Status = iree_vm_function_call_count_arguments_and_results(&Signature, &NumInputs, &NumOutputs);
+					TConstArrayView<UE::NNE::FTensorDesc> InputTensorDescs = InModule->GetFunctionMetaDataView()[0].InputDescs;
+					TConstArrayView<UE::NNE::FTensorDesc> OutputTensorDescs = InModule->GetFunctionMetaDataView()[0].OutputDescs;
+					if (!iree_status_is_ok(Status) || NumInputs != InputTensorDescs.Num() || NumOutputs != OutputTensorDescs.Num())
+					{
+						UE_LOG(LogNNERuntimeIREE, Error, TEXT("RDG session: Function signature mismatch in function %s"), *MainFunctionName);
+						iree_runtime_session_release(Session);
+						iree_status_free(Status);
+						return {};
+					}
+
+					iree_runtime_call_t Call;
+					Status = iree_runtime_call_initialize(Session, MainFunction, &Call);
+					if (!iree_status_is_ok(Status))
+					{
+						UE::NNERuntimeIREE::Private::PrintIREEError("RDG session: Failed to initialize the session call", Status);
+						iree_runtime_session_release(Session);
+						iree_status_free(Status);
+						return {};
+					}
+
+					TSharedPtr<FSession> Result = TSharedPtr<FSession>(new FSession(InDevice, InModule, Session, Call, InputTensorDescs, OutputTensorDescs));
+					iree_status_free(Status);
+					return Result;
+				}
+
+				TConstArrayView<UE::NNE::FTensorDesc> GetInputTensorDescs() const
+				{
+					return InputTensorDescs;
+				}
+
+				TConstArrayView<UE::NNE::FTensorDesc> GetOutputTensorDescs() const
+				{
+					return OutputTensorDescs;
+				}
+
+				TConstArrayView<UE::NNE::FTensorShape> GetInputTensorShapes() const
+				{
+					return InputTensorShapes;
+				}
+
+				TConstArrayView<UE::NNE::FTensorShape> GetOutputTensorShapes() const
+				{
+					return OutputTensorShapes;
+				}
+
+				ESetInputTensorShapesStatus SetInputTensorShapes(TConstArrayView<UE::NNE::FTensorShape> InInputShapes)
+				{
+					// OutputTensorShapes will be made available only if all shapes are concretes.
+					InputTensors.Reset();
+					OutputTensors.Reset();
+					OutputTensorShapes.Reset();
+					
+					bool bAllOutputShapeAreConcrete = true;
+					for (int32 i = 0; i < OutputTensorDescs.Num(); i++)
+					{
+						bAllOutputShapeAreConcrete &= OutputTensorDescs[i].GetShape().IsConcrete();
+					}
+					if (bAllOutputShapeAreConcrete)
+					{
+						for (int32 i = 0; i < OutputTensorDescs.Num(); i++)
+						{
+							NNE::Internal::FTensor Tensor = NNE::Internal::FTensor::MakeFromSymbolicDesc(OutputTensorDescs[i]);
+							OutputTensors.Add(Tensor);
+							OutputTensorShapes.Add(Tensor.GetShape());
+						}
+					}
+
+					InputTensorShapes.Reset(InInputShapes.Num());
+					if (InInputShapes.Num() != InputTensorDescs.Num())
+					{
+						UE_LOG(LogNNERuntimeIREE, Error, TEXT("CPU session: Number of input shapes does not match number of input tensors"));
+						return ESetInputTensorShapesStatus::Fail;
+					}
+
+					for (int32 i = 0; i < InInputShapes.Num(); ++i)
+					{
+						const UE::NNE::FTensorDesc SymbolicDesc = InputTensorDescs[i];
+						if (!InInputShapes[i].IsCompatibleWith(SymbolicDesc.GetShape()))
+						{
+							UE_LOG(LogNNERuntimeIREE, Error, TEXT("CPU session: Input shape does not match input tensor %s of index %d"), *SymbolicDesc.GetName(), i);
+							return ESetInputTensorShapesStatus::Fail;
+						}
+
+						NNE::Internal::FTensor Tensor = NNE::Internal::FTensor::Make(SymbolicDesc.GetName(), InInputShapes[i], SymbolicDesc.GetDataType());
+						InputTensors.Emplace(Tensor);
+					}
+					InputTensorShapes = InInputShapes;
+
+					return ESetInputTensorShapesStatus::Ok;
+				}
+
+				EEnqueueRDGStatus EnqueueRDG(FRDGBuilder& GraphBuilder, TConstArrayView<UE::NNE::FTensorBindingRDG> InInputBindings, TConstArrayView<UE::NNE::FTensorBindingRDG> InOutputBindings)
+				{
+					SCOPED_NAMED_EVENT_TEXT("NNERuntimeIREE::RDG::EnqueueRDG", FColor::Magenta);
+
+					// Verify the model inputs were prepared
+					if (InputTensorShapes.IsEmpty())
+					{
+						UE_LOG(LogNNERuntimeIREE, Error, TEXT("RDG session: Input shapes are not set, please call SetInputTensorShapes."));
+						return EEnqueueRDGStatus::Fail;
+					}
+
+					if (InInputBindings.Num() != InputTensorShapes.Num())
+					{
+						UE_LOG(LogNNERuntimeIREE, Error, TEXT("RDG session: Input bindings need to match input tensor descriptor count (got %d, expected %d)."), InInputBindings.Num(), InputTensorShapes.Num());
+						return EEnqueueRDGStatus::Fail;
+					}
+
+					if (!InOutputBindings.IsEmpty() && InOutputBindings.Num() != OutputTensorShapes.Num())
+					{
+						UE_LOG(LogNNERuntimeIREE, Error, TEXT("RDG session: Output binding can be empty or needs to match output tensor descriptor count (got %d, expected %d)."), InOutputBindings.Num(), OutputTensorShapes.Num());
+						return EEnqueueRDGStatus::Fail;
+					}
+
+					for (int32 i = 0; i < InInputBindings.Num(); i++)
+					{
+						const NNE::FTensorBindingRDG& Binding = InInputBindings[i];
+						if (!Binding.Buffer && InputTensors[i].GetDataSize() != 0)
+						{
+							UE_LOG(LogNNERuntimeIREE, Error, TEXT("RDG session: Binding input tensor %d is not set but given size by tensor descriptor is non-zero %d."), i, InputTensors[i].GetDataSize());
+							return EEnqueueRDGStatus::Fail;
+						}
+
+						if (Binding.Buffer && Binding.Buffer->Desc.GetSize() != InputTensors[i].GetDataSize())
+						{
+							UE_LOG(LogNNERuntimeIREE, Error, TEXT("RDG session: Binding input tensor %d size does not match size given by tensor descriptor (got %d, expected %d)."), i, Binding.Buffer->Desc.GetSize(), InputTensorShapes[i].Volume() * InputTensorDescs[i].GetElementByteSize());
+							return EEnqueueRDGStatus::Fail;
+						}
+					}
+
+					for (int32 i = 0; i < InOutputBindings.Num(); i++)
+					{
+						const NNE::FTensorBindingRDG& Binding = InOutputBindings[i];
+
+						if (Binding.Buffer && Binding.Buffer->Desc.GetSize() != OutputTensors[i].GetDataSize())
+						{
+							UE_LOG(LogNNERuntimeIREE, Error, TEXT("Binding output tensor %d size does not match tensor buffer size required by DirectML (got %d, expected %d)."), i, Binding.Buffer->Desc.GetSize(), OutputTensors[i].GetDataSize());
+							return EEnqueueRDGStatus::Fail;
+						}
+					}
+
+					iree_runtime_call_reset(&Call);
+
+					iree_status_t Status = iree_ok_status();
+
+					FGraphBuilderSetter GraphBuilderSetter(Device, GraphBuilder);
+
+					TArray<iree_hal_buffer_t*> IREEInputBuffers;
+					for (int i = 0; i < InInputBindings.Num(); i++)
+					{
+						iree_hal_buffer_t* IREEBuffer = NULL;
+						if (iree_status_is_ok(Status))
+						{
+							Status = IREE::HAL::RDG::BufferWrap(iree_allocator_system(), Device->GetDeviceAllocator(), InInputBindings[i].Buffer, &IREEBuffer);
+						}
+						if (iree_status_is_ok(Status))
+						{
+							IREEInputBuffers.Add(IREEBuffer);
+						}
+
+						TArray<iree_hal_dim_t> IREEShape;
+						for (uint32 Dim : InputTensorShapes[i].GetData())
+						{
+							IREEShape.Add(Dim);
+						}
+
+						iree_hal_buffer_view_t* IREEBufferView = NULL;
+						if (iree_status_is_ok(Status))
+						{
+							Status = iree_hal_buffer_view_create(
+								IREEBuffer,
+								IREEShape.Num(), IREEShape.GetData(),
+								IREE_HAL_ELEMENT_TYPE_FLOAT_32,
+								IREE_HAL_ENCODING_TYPE_DENSE_ROW_MAJOR,
+								iree_allocator_system(),
+								&IREEBufferView);
+						}
+
+						if (iree_status_is_ok(Status))
+						{
+							Status = iree_runtime_call_inputs_push_back_buffer_view(&Call, IREEBufferView);
+							iree_hal_buffer_view_release(IREEBufferView);
+						}
+
+						if (!iree_status_is_ok(Status))
+						{
+							break;
+						}
+					}
+
+					if (iree_status_is_ok(Status))
+					{
+						Status = iree_runtime_call_invoke(&Call, 0);
+					}
+
+					for (int i = 0; i < InOutputBindings.Num(); i++)
+					{
+						iree_hal_buffer_view_t* IREEBufferView = NULL;
+
+						IREE_CHECK_OK(iree_runtime_call_outputs_pop_front_buffer_view(&Call, &IREEBufferView));
+
+						iree_hal_buffer_t* IREEBuffer = iree_hal_buffer_view_buffer(IREEBufferView);
+
+						FRDGBufferRef RDGBuffer = IREE::HAL::RDG::BufferHandle(IREEBuffer);
+						check(RDGBuffer);
+
+						check(InOutputBindings[i].Buffer);
+						AddCopyBufferPass(GraphBuilder, InOutputBindings[i].Buffer, RDGBuffer);
+
+						iree_hal_buffer_view_release(IREEBufferView);
+					}
+
+					if (!iree_status_is_ok(Status))
+					{
+						NNERuntimeIREE::Private::PrintIREEError("EnqueueRDG Failed!", Status);
+					}
+
+					iree_status_free(Status);
+					return iree_status_is_ok(Status) ? EEnqueueRDGStatus::Ok : EEnqueueRDGStatus::Fail;
+				}
+
+			private:
+				TSharedRef<FDevice> Device;
+				TSharedRef<UE::NNERuntimeIREE::Private::FModule> Module;
+				iree_runtime_session_t* Session;
+				iree_runtime_call_t Call;
+				TArray<UE::NNE::FTensorDesc> InputTensorDescs;
+				TArray<UE::NNE::FTensorDesc> OutputTensorDescs;
+				TArray<UE::NNE::FTensorShape> InputTensorShapes;
+				TArray<UE::NNE::FTensorShape> OutputTensorShapes;
+				TArray<NNE::Internal::FTensor> InputTensors;
+				TArray<NNE::Internal::FTensor> OutputTensors;
+
+				FSession(TSharedRef<FDevice> InDevice, TSharedRef<UE::NNERuntimeIREE::Private::FModule> InModule, iree_runtime_session_t* InSession, const iree_runtime_call_t& InCall, TConstArrayView<UE::NNE::FTensorDesc> InInputTensorDescs, TConstArrayView<UE::NNE::FTensorDesc> InOutputTensorDescs)
+					: Device(InDevice), Module(InModule), Session(InSession), Call(InCall), InputTensorDescs(InInputTensorDescs), OutputTensorDescs(InOutputTensorDescs)
+				{
+					check(InSession);
+					check(!InputTensorDescs.IsEmpty());
+				}
+			};
+
+		} // namespace Private
+
+		FModelInstance::FModelInstance(TSharedRef<UE::NNERuntimeIREE::RDG::Private::FSession> InSession) : Session(InSession)
+		{
+			
+		}
+
+		TSharedPtr<FModelInstance> FModelInstance::Make(TSharedRef<UE::NNERuntimeIREE::RDG::Private::FDevice> InDevice, TSharedRef<UE::NNERuntimeIREE::Private::FModule> InModule)
+		{
+			TSharedPtr<Private::FSession> Session = Private::FSession::Make(InDevice, InModule);
+			if (!Session.IsValid())
+			{
+				return TSharedPtr<FModelInstance>();
+			}
+
+			return TSharedPtr<FModelInstance>(new FModelInstance(Session.ToSharedRef()));
+		}
+
+		TConstArrayView<UE::NNE::FTensorDesc> FModelInstance::GetInputTensorDescs() const
+		{
+			return Session->GetInputTensorDescs();
+		}
+
+		TConstArrayView<UE::NNE::FTensorDesc> FModelInstance::GetOutputTensorDescs() const
+		{
+			return Session->GetOutputTensorDescs();
+		}
+
+		TConstArrayView<UE::NNE::FTensorShape> FModelInstance::GetInputTensorShapes() const
+		{
+			return Session->GetInputTensorShapes();
+		}
+
+		TConstArrayView<UE::NNE::FTensorShape> FModelInstance::GetOutputTensorShapes() const
+		{
+			return Session->GetOutputTensorShapes();
+		}
+
+		FModelInstance::ESetInputTensorShapesStatus FModelInstance::SetInputTensorShapes(TConstArrayView<UE::NNE::FTensorShape> InInputShapes)
+		{
+			return Session->SetInputTensorShapes(InInputShapes);
+		}
+
+		FModelInstance::EEnqueueRDGStatus FModelInstance::EnqueueRDG(FRDGBuilder& GraphBuilder, TConstArrayView<UE::NNE::FTensorBindingRDG> InInputBindings, TConstArrayView<UE::NNE::FTensorBindingRDG> InOutputBindings)
+		{
+			return Session->EnqueueRDG(GraphBuilder, InInputBindings, InOutputBindings);
+		}
+
+		FModel::FModel(TSharedRef<UE::NNERuntimeIREE::RDG::Private::FDevice> InDevice, TSharedRef<UE::NNERuntimeIREE::Private::FModule> InModule) : Device(InDevice), Module(InModule)
+		{
+
+		}
+
+		TSharedPtr<FModel> FModel::Make(const FString& InDirPath, const FString& InVmfbFileName, const UNNERuntimeIREEModuleMetaData& InModuleMetaData)
+		{
+			check(!InVmfbFileName.IsEmpty());
+			check(!InModuleMetaData.FunctionMetaData.IsEmpty());
+
+			TSharedPtr<Private::FDevice> Device = Private::FDevice::Make();
+			if (!Device.IsValid())
+			{
+				return {};
+			}
+
+			TSharedPtr<UE::NNERuntimeIREE::Private::FModule> Module = UE::NNERuntimeIREE::Private::FModule::Make(InDirPath, InVmfbFileName, InModuleMetaData);
+			if (!Module.IsValid())
+			{
+				return {};
+			}
+
+			return TSharedPtr<FModel>(new FModel(Device.ToSharedRef(), Module.ToSharedRef()));
+		}
+
+		TSharedPtr<UE::NNE::IModelInstanceRDG> FModel::CreateModelInstanceRDG()
+		{
+			return FModelInstance::Make(Device, Module);
+		}
+
+	} // namespace RDG
+
 } // UE::NNERuntimeIREE
 
 #endif // WITH_NNE_RUNTIME_IREE
\ No newline at end of file
diff --git a/Engine/Plugins/Experimental/NNERuntimeIREE/Source/NNERuntimeIREE/Private/NNERuntimeIREEModel.h b/Engine/Plugins/Experimental/NNERuntimeIREE/Source/NNERuntimeIREE/Private/NNERuntimeIREEModel.h
--- a/Engine/Plugins/Experimental/NNERuntimeIREE/Source/NNERuntimeIREE/Private/NNERuntimeIREEModel.h
+++ b/Engine/Plugins/Experimental/NNERuntimeIREE/Source/NNERuntimeIREE/Private/NNERuntimeIREEModel.h
@@ -5,6 +5,7 @@
 #ifdef WITH_NNE_RUNTIME_IREE
 
 #include "NNERuntimeCPU.h"
+#include "NNERuntimeRDG.h"
 #include "NNERuntimeIREEMetaData.h"
 
 namespace UE::NNERuntimeIREE
@@ -65,6 +66,59 @@ namespace UE::NNERuntimeIREE
 			TSharedRef<UE::NNERuntimeIREE::Private::FModule> Module;
 		};
 	} // CPU
+
+	namespace RDG
+	{
+		namespace Private 
+		{
+			class FSession;
+			class FDevice;
+		} // Private
+
+		class FModelInstance : public UE::NNE::IModelInstanceRDG
+		{
+		private:
+			FModelInstance(TSharedRef<UE::NNERuntimeIREE::RDG::Private::FSession> InSession);
+
+		public:
+			static TSharedPtr<FModelInstance> Make(TSharedRef<UE::NNERuntimeIREE::RDG::Private::FDevice> InDevice, TSharedRef<UE::NNERuntimeIREE::Private::FModule> InModule);
+
+		public:
+			using ESetInputTensorShapesStatus = UE::NNE::IModelInstanceRDG::ESetInputTensorShapesStatus;
+			using EEnqueueRDGStatus = UE::NNE::IModelInstanceRDG::EEnqueueRDGStatus;
+
+			//~ Begin IModelInstanceRDG Interface
+			virtual TConstArrayView<UE::NNE::FTensorDesc> GetInputTensorDescs() const override;
+			virtual TConstArrayView<UE::NNE::FTensorDesc> GetOutputTensorDescs() const override;
+			virtual TConstArrayView<UE::NNE::FTensorShape> GetInputTensorShapes() const override;
+			virtual TConstArrayView<UE::NNE::FTensorShape> GetOutputTensorShapes() const override;
+			virtual ESetInputTensorShapesStatus SetInputTensorShapes(TConstArrayView<UE::NNE::FTensorShape> InInputShapes) override;
+			virtual EEnqueueRDGStatus EnqueueRDG(FRDGBuilder& RDGBuilder, TConstArrayView<UE::NNE::FTensorBindingRDG> InInputBindings, TConstArrayView<UE::NNE::FTensorBindingRDG> InOutputBindings) override;
+			//~ End IModelInstanceRDG Interface
+
+		private:
+			TSharedRef<UE::NNERuntimeIREE::RDG::Private::FSession> Session;
+		};
+
+		class FModel : public UE::NNE::IModelRDG
+		{
+		private:
+			FModel(TSharedRef<UE::NNERuntimeIREE::RDG::Private::FDevice> InDevice, TSharedRef<UE::NNERuntimeIREE::Private::FModule> InModule);
+
+		public:
+			static TSharedPtr<FModel> Make(const FString& InDirPath, const FString& InVmfbFileName, const UNNERuntimeIREEModuleMetaData& InModuleMetaData);
+
+		public:
+			//~ Begin IModelRDG Interface
+			virtual TSharedPtr<UE::NNE::IModelInstanceRDG> CreateModelInstanceRDG() override;
+			//~ End IModelRDG Interface
+
+		private:
+			TSharedRef<UE::NNERuntimeIREE::RDG::Private::FDevice> Device;
+			TSharedRef<UE::NNERuntimeIREE::Private::FModule> Module;
+		};
+	} // namespace RDG
+
 } // UE::NNERuntimeIREE
 
 #endif // WITH_NNE_RUNTIME_IREE
\ No newline at end of file
