[NNE] Fixed typo in output binding test in RunSyncCpu().

#tests Editor, Win64, ran models in NNETest

#rnx

diff --git a/Engine/Plugins/Experimental/NNERuntimeIREE/Source/NNERuntimeIREE/Private/NNERuntimeIREEModel.cpp b/Engine/Plugins/Experimental/NNERuntimeIREE/Source/NNERuntimeIREE/Private/NNERuntimeIREEModel.cpp
--- a/Engine/Plugins/Experimental/NNERuntimeIREE/Source/NNERuntimeIREE/Private/NNERuntimeIREEModel.cpp
+++ b/Engine/Plugins/Experimental/NNERuntimeIREE/Source/NNERuntimeIREE/Private/NNERuntimeIREEModel.cpp
@@ -780,9 +780,9 @@ namespace UE::NNERuntimeIREE
 						return ERunSyncStatus::Fail;
 					}
 
-					if (!InOutputBindings.IsEmpty() && InOutputBindings.Num() != OutputTensorShapes.Num())
+					if (!InOutputBindings.IsEmpty() && InOutputBindings.Num() != OutputTensorDescs.Num())
 					{
-						UE_LOG(LogNNERuntimeIREE, Error, TEXT("CPU session: Output binding can be empty or needs to match output tensor descriptor count (got %d, expected %d)."), InOutputBindings.Num(), OutputTensorShapes.Num());
+						UE_LOG(LogNNERuntimeIREE, Error, TEXT("CPU session: Output binding can be empty or needs to match output tensor descriptor count (got %d, expected %d)."), InOutputBindings.Num(), OutputTensorDescs.Num());
 						return ERunSyncStatus::Fail;
 					}
 
