[NNE] NNERuntimeIREERdg add shader compilation code (wip, runtime is not available).

#test Editor, Win64, tested add.mlir

#rnx

diff --git a/Engine/Plugins/Experimental/NNERuntimeIREE/NNERuntimeIREE.uplugin b/Engine/Plugins/Experimental/NNERuntimeIREE/NNERuntimeIREE.uplugin
--- a/Engine/Plugins/Experimental/NNERuntimeIREE/NNERuntimeIREE.uplugin
+++ b/Engine/Plugins/Experimental/NNERuntimeIREE/NNERuntimeIREE.uplugin
@@ -31,7 +31,7 @@
 			"LoadingPhase": "Default"
 		},
 		{
-			"Name": "IREEDriverRDGShaders",
+			"Name": "NNERuntimeIREEShader",
 			"Type": "Runtime",
 			"LoadingPhase": "PostConfigInit"
 		}
diff --git a/Engine/Plugins/Experimental/NNERuntimeIREE/Shaders/NNERuntimeIREEAdd.usf b/Engine/Plugins/Experimental/NNERuntimeIREE/Shaders/NNERuntimeIREEAdd.usf
deleted file mode 100644
diff --git a/Engine/Plugins/Experimental/NNERuntimeIREE/Shaders/NNERuntimeIREEAdd_Cross.usf b/Engine/Plugins/Experimental/NNERuntimeIREE/Shaders/NNERuntimeIREEAdd_Cross.usf
deleted file mode 100644
diff --git a/Engine/Plugins/Experimental/NNERuntimeIREE/Shaders/NNERuntimeIREEShader.usf b/Engine/Plugins/Experimental/NNERuntimeIREE/Shaders/NNERuntimeIREEShader.usf
new file mode 100644
--- /dev/null
+++ b/Engine/Plugins/Experimental/NNERuntimeIREE/Shaders/NNERuntimeIREEShader.usf
@@ -0,0 +1,6 @@
+// Copyright Epic Games, Inc. All Rights Reserved.
+
+#include "/Engine/Public/Platform.ush"
+
+// Include the generated kernel source.
+#include "/Plugin/NNERuntimeIREEShader/Generated/NNERuntimeIREEShader.ush"
\ No newline at end of file
diff --git a/Engine/Plugins/Experimental/NNERuntimeIREE/Source/IREEDriverRDG/IREEDriverRDG.Build.cs b/Engine/Plugins/Experimental/NNERuntimeIREE/Source/IREEDriverRDG/IREEDriverRDG.Build.cs
--- a/Engine/Plugins/Experimental/NNERuntimeIREE/Source/IREEDriverRDG/IREEDriverRDG.Build.cs
+++ b/Engine/Plugins/Experimental/NNERuntimeIREE/Source/IREEDriverRDG/IREEDriverRDG.Build.cs
@@ -20,10 +20,21 @@ public class IREEDriverRDG : ModuleRules
 				"Projects",
 				"RenderCore",
 				"RHI",
-				"IREEDriverRDGShaders" // TOTO remove
+				"NNERuntimeIREEShader"
 			}
 		);
 
+		if (Target.Type == TargetType.Editor)
+		{
+			PrivateDependencyModuleNames.AddRange(
+				new string[]
+				{
+					"Json",
+					"ShaderCompilerCommon"
+				}
+			);
+		}
+
 		if (Target.Platform == UnrealTargetPlatform.Win64)
 		{
 			PublicDefinitions.Add("WITH_IREE_DRIVER_RDG");
diff --git a/Engine/Plugins/Experimental/NNERuntimeIREE/Source/IREEDriverRDG/Internal/IREECompilerRDG.h b/Engine/Plugins/Experimental/NNERuntimeIREE/Source/IREEDriverRDG/Internal/IREECompilerRDG.h
new file mode 100644
--- /dev/null
+++ b/Engine/Plugins/Experimental/NNERuntimeIREE/Source/IREEDriverRDG/Internal/IREECompilerRDG.h
@@ -0,0 +1,96 @@
+// Copyright Epic Games, Inc. All Rights Reserved.
+
+#pragma once
+
+#include "CoreMinimal.h"
+// #include "NNERuntimeIREEMetaData.h"
+#include "Serialization/JsonSerializerMacros.h"
+
+#include "IREECompilerRDG.generated.h"
+
+USTRUCT()
+struct FIREECompileRDGArchitectureInfoCPU
+{
+	GENERATED_BODY()
+
+	UPROPERTY()
+	FString Architecture;
+
+	UPROPERTY()
+	FString RelativeDirPath;
+	
+	UPROPERTY()
+	FString SharedLibraryFileName;
+
+	UPROPERTY()
+	FString VmfbFileName;
+
+	UPROPERTY()
+	FString SharedLibraryEntryPointName;
+};
+
+USTRUCT()
+struct FIREECompileRDGCompilerResultCPU
+{
+	GENERATED_BODY()
+
+	UPROPERTY()
+	TArray<FIREECompileRDGArchitectureInfoCPU> ArchitectureInfos;
+};
+
+#ifdef WITH_IREE_DRIVER_RDG
+#if WITH_EDITOR
+
+namespace UE::IREE::Compiler::RDG
+{
+	struct FBuildTarget : public FJsonSerializable
+	{
+		FString Architecture;
+		FString CompilerArguments;
+		FString LinkerArguments;
+
+		BEGIN_JSON_SERIALIZER
+			JSON_SERIALIZE("Architecture", Architecture);
+			JSON_SERIALIZE("CompilerArguments", CompilerArguments);
+			JSON_SERIALIZE("LinkerArguments", LinkerArguments);
+		END_JSON_SERIALIZER
+	};
+
+	struct FBuildConfig : public FJsonSerializable
+	{
+		TArray<FString> CompilerCommand;
+		TArray<FString> LinkerCommand;
+		FString SharedLibExt;
+		TArray<FBuildTarget> BuildTargets;
+
+		BEGIN_JSON_SERIALIZER
+			JSON_SERIALIZE_ARRAY("CompilerCommand", CompilerCommand);
+			JSON_SERIALIZE_ARRAY("LinkerCommand", LinkerCommand);
+			JSON_SERIALIZE("SharedLibExt", SharedLibExt);
+			JSON_SERIALIZE_ARRAY_SERIALIZABLE("Targets", BuildTargets, FBuildTarget);
+		END_JSON_SERIALIZER
+	};
+
+	class IREEDRIVERRDG_API FCompiler
+	{
+	private:
+		FCompiler(const ITargetPlatform* TargetPlatform, const FString& InCompilerCommand, const FString& InLinkerCommand, const FString& InSharedLibExt, TConstArrayView<FBuildTarget> InBuildTargets);
+
+	public:
+		~FCompiler() {};
+		static TUniquePtr<FCompiler> Make(const ITargetPlatform* TargetPlatform);
+
+	public:
+		bool CompileMlir(TConstArrayView<uint8> InFileData, const FString& InModelName, const FString& InOutputDir, FIREECompileRDGCompilerResultCPU& OutCompilerResult);
+
+	private:
+		const ITargetPlatform* TargetPlatform;
+		FString CompilerCommand;
+		FString LinkerCommand;
+		FString SharedLibExt;
+		TArray<FBuildTarget> BuildTargets;
+	};
+} // namespace UE::IREE::Compiler::RDG
+
+#endif // WITH_EDITOR
+#endif // WITH_IREE_DRIVER_RDG
\ No newline at end of file
diff --git a/Engine/Plugins/Experimental/NNERuntimeIREE/Source/IREEDriverRDG/Internal/IREEDriverRDG.h b/Engine/Plugins/Experimental/NNERuntimeIREE/Source/IREEDriverRDG/Internal/IREEDriverRDG.h
--- a/Engine/Plugins/Experimental/NNERuntimeIREE/Source/IREEDriverRDG/Internal/IREEDriverRDG.h
+++ b/Engine/Plugins/Experimental/NNERuntimeIREE/Source/IREEDriverRDG/Internal/IREEDriverRDG.h
@@ -19,10 +19,11 @@ THIRD_PARTY_INCLUDES_END
 #endif // PLATFORM_MICROSOFT
 
 #include "RenderGraphFwd.h"
+
 namespace UE::IREE::HAL::RDG
 {
 
-IREEDRIVERRDG_API iree_status_t DeviceCreate(iree_allocator_t HostAllocator, iree_hal_device_t** OutDevice);
+IREEDRIVERRDG_API iree_status_t DeviceCreate(iree_allocator_t HostAllocator, const FString& BasePath, iree_hal_device_t** OutDevice);
 
 IREEDRIVERRDG_API iree_status_t BufferWrap(iree_allocator_t HostAllocator, iree_hal_allocator_t* DeviceAllocator, FRDGBuffer* InBuffer, iree_hal_buffer_t** OutBuffer);
 
diff --git a/Engine/Plugins/Experimental/NNERuntimeIREE/Source/IREEDriverRDG/Private/IREECompilerRDG.cpp b/Engine/Plugins/Experimental/NNERuntimeIREE/Source/IREEDriverRDG/Private/IREECompilerRDG.cpp
new file mode 100644
--- /dev/null
+++ b/Engine/Plugins/Experimental/NNERuntimeIREE/Source/IREEDriverRDG/Private/IREECompilerRDG.cpp
@@ -0,0 +1,439 @@
+// Copyright Epic Games, Inc. All Rights Reserved.
+
+#include "IREECompilerRDG.h"
+
+#ifdef WITH_IREE_DRIVER_RDG
+#if WITH_EDITOR
+
+// SPIR-V to HLSL
+#include "CrossCompiler.h"
+#include "CrossCompilerCommon.h"
+#include "ShaderCompilerCommon.h"
+#include "ShaderCompilerDefinitions.h"
+
+// UE shader compilation
+#include "NNERuntimeIREEShaderShared.h"
+#include "ShaderParameterMetadataBuilder.h"
+
+#include "GenericPlatform/GenericPlatformFile.h"
+#include "GenericPlatform/GenericPlatformMisc.h"
+#include "GenericPlatform/GenericPlatformProcess.h"
+#include "HAL/FileManager.h"
+#include "HAL/PlatformFileManager.h"
+#include "Interfaces/IPluginManager.h"
+#include "Interfaces/ITargetPlatform.h"
+#include "IREEDriverRDGLog.h"
+#include "Kismet/GameplayStatics.h"
+#include "Misc/FileHelper.h"
+#include "Misc/MonitoredProcess.h"
+#include "Misc/Paths.h"
+#include "NNE.h"
+#include "Serialization/JsonReader.h"
+#include "Serialization/JsonSerializer.h"
+#include "Templates/SharedPointer.h"
+
+namespace UE::IREE::Compiler
+{
+	namespace RDG
+	{
+		namespace Private
+		{
+			bool ResolveEnvironmentVariables(FString& String)
+			{
+				const FString StartString = "$ENV{";
+				const FString EndString = "}";
+				FString ResultString = String;
+				int32 StartIndex = ResultString.Find(StartString, ESearchCase::CaseSensitive);
+				while (StartIndex != INDEX_NONE)
+				{
+					StartIndex += StartString.Len();
+					int32 EndIndex = ResultString.Find(EndString, ESearchCase::CaseSensitive, ESearchDir::FromStart, StartIndex);
+					if (EndIndex > StartIndex)
+					{
+						FString EnvironmentVariableName = ResultString.Mid(StartIndex, EndIndex - StartIndex);
+						FString EnvironmentVariableValue = FPlatformMisc::GetEnvironmentVariable(*EnvironmentVariableName);
+						if (EnvironmentVariableValue.IsEmpty())
+						{
+							return false;
+						}
+						else
+						{
+							ResultString.ReplaceInline(*(StartString + EnvironmentVariableName + EndString), *EnvironmentVariableValue, ESearchCase::CaseSensitive);
+						}
+					}
+					else
+					{
+						return false;
+					}
+					StartIndex = ResultString.Find(StartString, ESearchCase::CaseSensitive);
+				}
+				String = ResultString;
+				return true;
+			}
+
+			FString GetSharedLibraryEntryPointName(const FString& HeaderString)
+			{
+				FString SearchString = "iree_hal_executable_library_header_t**";
+				int32 Start = HeaderString.Find(SearchString);
+				if (Start == INDEX_NONE)
+				{
+					return "";
+				}
+				Start += SearchString.Len();
+				int32 End = HeaderString.Find("(", ESearchCase::CaseSensitive, ESearchDir::FromStart, Start);
+				if (End <= Start)
+				{
+					return "";
+				}
+				return HeaderString.Mid(Start, End - Start).TrimStartAndEnd();
+			}
+
+			void RunCommand(const FString& Command, const FString& Arguments, const FString& LogFilePath = FString())
+			{
+				int32 ReturnCode = 0;
+				bool IsCanceled = false;
+
+				FMonitoredProcess Process(Command, Arguments, true);
+				Process.OnCompleted().BindLambda([&ReturnCode] (int32 _ReturnCode) { ReturnCode = _ReturnCode; });
+				Process.OnCanceled().BindLambda([&IsCanceled] (){ IsCanceled = true; });
+
+				if (!Process.Launch())
+				{
+					UE_LOG(LogIREEDriverRDG, Warning, TEXT("Failed to launch subprocess!"));
+					return;
+				}
+
+				while (Process.Update())
+				{
+					// Poll until process has finished
+				}
+
+				if (IsCanceled)
+				{
+					UE_LOG(LogIREEDriverRDG, Warning, TEXT("Execution of subprocess was canceled!"));
+				}
+				else if (ReturnCode)
+				{
+					UE_LOG(LogIREEDriverRDG, Warning, TEXT("Subprocess exited with non-zero code %d"), ReturnCode);
+				}
+
+				if (!LogFilePath.IsEmpty())
+				{
+					FFileHelper::SaveStringToFile(Process.GetFullOutputWithoutDelegate(), *LogFilePath);
+
+					UE_LOG(LogIREEDriverRDG, Log, TEXT("Saved subprocess output to: %s"), *LogFilePath);
+				}
+			}
+
+			bool ConvertSpirvToHLSL(const FString& Filepath, const FString& EntryPointName)
+			{
+				FString OutputFilepath = FPaths::GetBaseFilename(Filepath, false) + ".hlsl";
+
+				TArray<uint8> SpirvCode;
+				if (!FFileHelper::LoadFileToArray(SpirvCode, *Filepath))
+				{
+					UE_LOG(LogIREEDriverRDG, Log, TEXT("Could not load file to array: %s"), *Filepath);
+					return false;
+				}
+
+				CrossCompiler::FShaderConductorContext CompilerContext;
+
+				FString ShaderSource;
+				FString VirtualSourceFilePath;
+				FShaderCompilerDefinitions AdditionalDefines;
+
+				if (!CompilerContext.LoadSource(ShaderSource, VirtualSourceFilePath, EntryPointName, EShaderFrequency::SF_Compute, &AdditionalDefines))
+				{
+					UE_LOG(LogIREEDriverRDG, Log, TEXT("Shader condutor could not load source"));
+					return false;
+				}
+
+				CrossCompiler::FShaderConductorOptions Options;
+				CrossCompiler::FShaderConductorTarget Target;
+				Target.Language = CrossCompiler::EShaderConductorLanguage::Hlsl;
+				Target.Version = 62;
+
+				FString HlslSource;
+				bool Result = false;
+				try
+				{
+					Result = CompilerContext.CompileSpirvToSource(Options, Target, SpirvCode.GetData(), SpirvCode.Num(), HlslSource);
+				}
+				catch(const std::exception& Exception)
+				{
+					UE_LOG(LogIREEDriverRDG, Log, TEXT("Error when compiling Spirv to source %hs"), Exception.what());
+					return false;
+				}
+				
+				
+				if (Result)
+				{
+					if (!FFileHelper::SaveStringToFile(HlslSource, *OutputFilepath))
+					{
+						UE_LOG(LogIREEDriverRDG, Log, TEXT("Could not save HLSL source string to file %s"), *OutputFilepath);
+						return false;
+					}
+				}
+
+				return Result;
+			}
+
+			bool CompileShader(const FString& Filepath, const FString& EntryPointName, const ITargetPlatform* TargetPlatform)
+			{
+				FString OutputFilepath = FPaths::GetBaseFilename(Filepath, false) + ".ireeshadermap";
+
+				FString HlslSource;
+				if (!FFileHelper::LoadFileToString(HlslSource, *Filepath))
+				{
+					UE_LOG(LogIREEDriverRDG, Log, TEXT("Could not load file to HLSL source string: %s"), *Filepath);
+					return false;
+				}
+
+				const ERHIFeatureLevel::Type CacheFeatureLevel = GMaxRHIFeatureLevel;
+				const EShaderPlatform ShaderPlatform = GShaderPlatformForFeatureLevel[CacheFeatureLevel];
+				FShaderParametersMetadata* ShaderParameterMetadata = nullptr;
+				{
+					FShaderParametersMetadataBuilder Builder;
+					Builder.AddRDGBufferUAV(TEXT("_resource_var_0_0_"), TEXT("RWByteAddressBuffer"));
+					Builder.AddRDGBufferUAV(TEXT("_resource_var_0_1_"), TEXT("RWByteAddressBuffer"));
+					Builder.AddRDGBufferUAV(TEXT("_resource_var_0_2_"), TEXT("RWByteAddressBuffer"));
+
+					ShaderParameterMetadata = Builder.Build(FShaderParametersMetadata::EUseCase::ShaderParameterStruct, TEXT("iree_add_parameters"));
+				}
+
+				TUniquePtr<FNNERuntimeIREEResource> KernelResource = MakeUnique<FNNERuntimeIREEResource>();
+				KernelResource->SetupResource(
+					CacheFeatureLevel,
+					EntryPointName, // TODO include module name,
+					EntryPointName,
+					FString(),
+					HlslSource,
+					ShaderParameterMetadata,
+					FName()
+				);
+
+				if (!KernelResource->CacheShaders(ShaderPlatform, TargetPlatform, true, true))
+				{
+					UE_LOG(LogIREEDriverRDG, Warning, TEXT("Failed to compile FNNERuntimeIREEResource [%s] for platform [%s]."), *KernelResource->GetFriendlyName(), *LegacyShaderPlatformToShaderFormat(ShaderPlatform).ToString());
+				}
+
+				TArray/*64*/<uint8> ResultData;
+				FMemoryWriter/*64*/ Writer(ResultData, /*bIsPersitent =*/ true);
+
+				KernelResource->SerializeShaderMap(Writer);
+
+				FFileHelper::SaveArrayToFile(ResultData, *OutputFilepath);
+
+				return true;
+			}
+		} // Private
+
+		FCompiler::FCompiler(const ITargetPlatform* InTargetPlatform, const FString& InCompilerCommand, const FString& InLinkerCommand, const FString& InSharedLibExt, TConstArrayView<FBuildTarget> InBuildTargets) : TargetPlatform(InTargetPlatform), CompilerCommand(InCompilerCommand), LinkerCommand(InLinkerCommand), SharedLibExt(InSharedLibExt), BuildTargets(InBuildTargets)
+		{
+			
+		}
+
+		TUniquePtr<FCompiler> FCompiler::Make(const ITargetPlatform* TargetPlatform)
+		{
+			using namespace Private;
+
+			const FString TargetPlatformName = TargetPlatform ? TargetPlatform->IniPlatformName() : UGameplayStatics::GetPlatformName();
+
+			FString PluginDir = FPaths::ConvertRelativePathToFull(*IPluginManager::Get().FindPlugin(UE_PLUGIN_NAME)->GetBaseDir());
+			FString BuildConfigFileName = FString("IREE_") + UGameplayStatics::GetPlatformName() + "_To_" + TargetPlatformName + ".json";
+			TArray<FString> BuildConfigFilePaths =
+			{
+				FPaths::Combine(FPaths::ConvertRelativePathToFull(*FPaths::ProjectConfigDir()), BuildConfigFileName),
+				FPaths::Combine(PluginDir, "Config", BuildConfigFileName),
+				FPaths::Combine(FPaths::ConvertRelativePathToFull(*FPaths::EngineDir()), "Platforms", TargetPlatformName, "Plugins", UE_PLUGIN_NAME, "Config", BuildConfigFileName),
+				FPaths::Combine(FPaths::ConvertRelativePathToFull(*FPaths::EngineDir()), "Platforms", TargetPlatformName, "Plugins", "Experimental", UE_PLUGIN_NAME, "Config", BuildConfigFileName)
+			};
+
+			FString CompilerCommand;
+			FString LinkerCommand;
+			FString SharedLibExt;
+			TArray<FBuildTarget> BuildTargets;
+			IPlatformFile& PlatformFile = FPlatformFileManager::Get().GetPlatformFile();
+			for (FString BuildConfigFilePath : BuildConfigFilePaths)
+			{
+				if (PlatformFile.FileExists(*BuildConfigFilePath))
+				{
+					FString BuildConfigFileString;
+					if (FFileHelper::LoadFileToString(BuildConfigFileString, *BuildConfigFilePath))
+					{
+						TSharedRef<TJsonReader<TCHAR>> JsonReader = TJsonReaderFactory<TCHAR>::Create(BuildConfigFileString);
+						TSharedPtr<FJsonObject> JsonObject = MakeShareable(new FJsonObject);
+						FBuildConfig BuildConfig;
+						if (FJsonSerializer::Deserialize(JsonReader, JsonObject) && JsonObject.IsValid() && BuildConfig.FromJson(JsonObject))
+						{
+							if (BuildConfig.BuildTargets.IsEmpty())
+							{
+								UE_LOG(LogIREEDriverRDG, Warning, TEXT("IREECompilerRDG could not find targets in %s"), *BuildConfigFilePath);
+								continue;
+							}
+
+							FString TmpCompilerCommand;
+							for (int32 i = 0; i < BuildConfig.CompilerCommand.Num(); i++)
+							{
+								if (ResolveEnvironmentVariables(BuildConfig.CompilerCommand[i]))
+								{
+									BuildConfig.CompilerCommand[i].ReplaceInline(*FString("${PLUGIN_DIR}"), *PluginDir);
+									BuildConfig.CompilerCommand[i].ReplaceInline(*FString("${PROJECT_DIR}"), *FPaths::ProjectDir());
+									if (PlatformFile.FileExists(*BuildConfig.CompilerCommand[i]))
+									{
+										TmpCompilerCommand = BuildConfig.CompilerCommand[i];
+										break;
+									}
+								}
+								else
+								{
+									UE_LOG(LogIREEDriverRDG, Warning, TEXT("IREECompilerRDG could not replace environment variables in %s"), *BuildConfig.CompilerCommand[i]);
+								}
+							}
+							if (TmpCompilerCommand.IsEmpty())
+							{
+								UE_LOG(LogIREEDriverRDG, Warning, TEXT("IREECompilerRDG could not find the compiler executable in %s"), *BuildConfigFilePath);
+								continue;
+							}
+
+							FString TmpLinkerCommand;
+							for (int32 i = 0; i < BuildConfig.LinkerCommand.Num(); i++)
+							{
+								if (ResolveEnvironmentVariables(BuildConfig.LinkerCommand[i]))
+								{
+									BuildConfig.LinkerCommand[i].ReplaceInline(*FString("${PLUGIN_DIR}"), *PluginDir);
+									BuildConfig.LinkerCommand[i].ReplaceInline(*FString("${PROJECT_DIR}"), *FPaths::ProjectDir());
+									if (PlatformFile.FileExists(*BuildConfig.LinkerCommand[i]))
+									{
+										TmpLinkerCommand = BuildConfig.LinkerCommand[i];
+										break;
+									}
+								}
+								else
+								{
+									UE_LOG(LogIREEDriverRDG, Warning, TEXT("IREECompilerRDG could not replace environment variables in %s"), *BuildConfig.LinkerCommand[i]);
+								}
+							}
+							if (TmpLinkerCommand.IsEmpty())
+							{
+								UE_LOG(LogIREEDriverRDG, Warning, TEXT("IREECompilerRDG could not find the linker executable in %s"), *BuildConfigFilePath);
+								continue;
+							}
+
+							CompilerCommand = TmpCompilerCommand;
+							LinkerCommand = TmpLinkerCommand;
+							SharedLibExt = BuildConfig.SharedLibExt;
+							BuildTargets = BuildConfig.BuildTargets;
+							break;
+						}
+						else
+						{
+							UE_LOG(LogIREEDriverRDG, Warning, TEXT("IREECompilerRDG could not parse build config file %s"), *BuildConfigFilePath);
+						}
+					}
+					else
+					{
+						UE_LOG(LogIREEDriverRDG, Warning, TEXT("IREECompilerRDG could not read build config file %s"), *BuildConfigFilePath);
+					}
+				}
+			}
+			if (CompilerCommand.IsEmpty() || LinkerCommand.IsEmpty() || BuildTargets.IsEmpty())
+			{
+				return TUniquePtr<FCompiler>();
+			}
+			return TUniquePtr<FCompiler>(new FCompiler(TargetPlatform, CompilerCommand, LinkerCommand, SharedLibExt, BuildTargets));
+		}
+
+		bool FCompiler::CompileMlir(TConstArrayView<uint8> InFileData, const FString& InModelName, const FString& InOutputDir, FIREECompileRDGCompilerResultCPU& OutCompilerResult)
+		{
+			using namespace Private;
+
+			IPlatformFile& PlatformFile = FPlatformFileManager::Get().GetPlatformFile();
+
+			FString InputFilePath = FPaths::Combine(InOutputDir, InModelName) + ".mlir";
+			if (!PlatformFile.FileExists(*InputFilePath))
+			{
+				SCOPED_NAMED_EVENT_TEXT("InputFile", FColor::Magenta);
+
+				FFileHelper::SaveArrayToFile(InFileData, *InputFilePath);
+			}
+
+			bool bResult = true;
+			for (int32 i = 0; i < BuildTargets.Num(); i++)
+			{
+				FString IntermediateDirPath = FPaths::Combine(InOutputDir, BuildTargets[i].Architecture);
+				PlatformFile.CreateDirectoryTree(*IntermediateDirPath);
+				FString IntermediateFilePathNoExt = FPaths::Combine(IntermediateDirPath, InModelName);
+				FString VmfbFilePath = IntermediateFilePathNoExt + ".vmfb";
+
+				FString CompilerArguments = TEXT("--iree-hal-target-backends=vulkan-spirv --iree-hal-dump-executable-binaries-to=${BINARIES_PATH} -o ${VMFB_PATH} ${INPUT_PATH}");
+				CompilerArguments.ReplaceInline(*FString("${BINARIES_PATH}"), *(FString("\"") + IntermediateDirPath + "\""));
+				CompilerArguments.ReplaceInline(*FString("${VMFB_PATH}"), *(FString("\"") + VmfbFilePath + "\""));
+				CompilerArguments.ReplaceInline(*FString("${INPUT_PATH}"), *(FString("\"") + InputFilePath + "\""));
+
+				RunCommand(CompilerCommand, CompilerArguments, IntermediateFilePathNoExt + "_compile-log.txt");
+
+				if (!PlatformFile.FileExists(*VmfbFilePath))
+				{
+					UE_LOG(LogIREEDriverRDG, Warning, TEXT("IREECompilerRDG failed to compile the model \"%s\" using the command:"), *InputFilePath);
+					UE_LOG(LogIREEDriverRDG, Warning, TEXT("\"%s\" %s"), *CompilerCommand, *CompilerArguments);
+					bResult = false;
+					continue;
+				}
+
+				// TODO get entry point name from IREE compiler
+				const FString SpirvEntryPointName = TEXT("add_dispatch_0_generic_4_f32");
+
+				TArray<FString> SpirvFilenames;
+				PlatformFile.FindFiles(SpirvFilenames, *IntermediateDirPath, TEXT(".spv"));
+				if (SpirvFilenames.IsEmpty())
+				{
+					UE_LOG(LogIREEDriverRDG, Warning, TEXT("No SPIR-V files generated."));
+				}
+				for (const FString& Filename : SpirvFilenames)
+				{
+					UE_LOG(LogIREEDriverRDG, Log, TEXT("Spirv %s"), *Filename);
+
+					ConvertSpirvToHLSL(Filename, SpirvEntryPointName);
+				}
+
+				const FString HlslEntryPointName = TEXT("main");
+
+				TArray<FString> HlslFilenames;
+				PlatformFile.FindFiles(HlslFilenames, *IntermediateDirPath, TEXT(".hlsl"));
+				check(HlslFilenames.Num() == SpirvFilenames.Num());
+
+				for (const FString& Filename : HlslFilenames)
+				{
+
+					CompileShader(Filename, HlslEntryPointName, TargetPlatform);
+				}
+
+				TArray<FString> ShaderMapFilenames;
+				PlatformFile.FindFiles(ShaderMapFilenames, *IntermediateDirPath, TEXT(".ireeshadermap"));
+				check(ShaderMapFilenames.Num() == SpirvFilenames.Num());
+
+				check(ShaderMapFilenames.Num() == 1); // TODO
+
+				FIREECompileRDGArchitectureInfoCPU ArchitectureInfo;
+				ArchitectureInfo.Architecture = BuildTargets[i].Architecture;
+				ArchitectureInfo.RelativeDirPath = BuildTargets[i].Architecture;
+				ArchitectureInfo.VmfbFileName = InModelName + ".vmfb";
+				ArchitectureInfo.SharedLibraryFileName = FPaths::GetCleanFilename(ShaderMapFilenames[0]);
+				OutCompilerResult.ArchitectureInfos.Add(MoveTemp(ArchitectureInfo));
+			}
+
+			bResult &= !OutCompilerResult.ArchitectureInfos.IsEmpty();
+			
+			if (!bResult)
+			{
+				OutCompilerResult.ArchitectureInfos.Empty();
+			}
+			
+			return bResult;
+		}
+	} // namespace RDG
+} // namespace UE::IREE::Compiler
+
+#endif // WITH_EDITOR
+#endif // WITH_IREE_DRIVER_RDG
\ No newline at end of file
diff --git a/Engine/Plugins/Experimental/NNERuntimeIREE/Source/IREEDriverRDG/Private/IREEDriverRDGCommandBuffer.cpp b/Engine/Plugins/Experimental/NNERuntimeIREE/Source/IREEDriverRDG/Private/IREEDriverRDGCommandBuffer.cpp
--- a/Engine/Plugins/Experimental/NNERuntimeIREE/Source/IREEDriverRDG/Private/IREEDriverRDGCommandBuffer.cpp
+++ b/Engine/Plugins/Experimental/NNERuntimeIREE/Source/IREEDriverRDG/Private/IREEDriverRDGCommandBuffer.cpp
@@ -5,12 +5,14 @@
 #ifdef WITH_IREE_DRIVER_RDG
 
 #include "IREEDriverRDGBuffer.h"
+#include "IREEDriverRDGExecutable.h"
 #include "IREEDriverRDGLog.h"
+#include "NNERuntimeIREEShader.h"
+#include "NNERuntimeIREEShaderShared.h"
 #include "RenderGraphBuilder.h"
 #include "RenderGraphUtils.h"
 
-
-#include "IREEDriverRDGShadersAddCrossCS.h"
+#include "ShaderParameterMetadataBuilder.h"
 
 DECLARE_GPU_STAT_NAMED(FCommandBufferDispatch, TEXT("CommandBuffer.Dispatch"));
 namespace UE::IREE::HAL::RDG
@@ -69,7 +71,7 @@ public:
 
 	iree_status_t Dispatch(iree_hal_executable_t* Executable, int32 EntryPoint, uint32 WorkgroupX, uint32 WorkgroupY, uint32 WorkgroupZ)
 	{
-		using namespace UE::IREEDriverRDGShaders::Internal;
+		// using namespace UE::IREEDriverRDGShaders::Internal;
 
 		iree_hal_allocator_t* DeviceAllocator = iree_hal_device_allocator(Device);
 		FRDGBuilder& GraphBuilder = DeviceAllocatorGetGraphBuilder(DeviceAllocator); // TODO 
@@ -88,13 +90,37 @@ public:
 			AddCopyBufferPass(GraphBuilder, ACopy, RecordedBindings[0].Buffer);
 			AddCopyBufferPass(GraphBuilder, BCopy, RecordedBindings[1].Buffer);
 
-			FAddCrossCS::FParameters *ShaderParameters = GraphBuilder.AllocParameters<FAddCrossCS::FParameters>();
-			ShaderParameters->_resource_var_0_0_ = GraphBuilder.CreateUAV(ACopy);
-			ShaderParameters->_resource_var_0_1_ = GraphBuilder.CreateUAV(BCopy);
-			ShaderParameters->_resource_var_0_2_ = GraphBuilder.CreateUAV(ResultCopy);
-
-			FGlobalShaderMap* GlobalShaderMap = GetGlobalShaderMap(GMaxRHIFeatureLevel);
-			TShaderMapRef<FAddCrossCS> Shader(GlobalShaderMap);
+			const FNNERuntimeIREEResource& KernelResource = ExecutableGetResource(Executable);
+			TShaderRef<FNNERuntimeIREEShader> Shader = KernelResource.GetShader(0);
+			check(Shader.IsValid());
+
+			FShaderParametersMetadata* ShaderParameterMetadata = nullptr;
+			{
+				FShaderParametersMetadataBuilder Builder;
+				Builder.AddRDGBufferUAV(TEXT("_resource_var_0_0_"), TEXT("RWByteAddressBuffer"));
+				Builder.AddRDGBufferUAV(TEXT("_resource_var_0_1_"), TEXT("RWByteAddressBuffer"));
+				Builder.AddRDGBufferUAV(TEXT("_resource_var_0_2_"), TEXT("RWByteAddressBuffer"));
+
+				ShaderParameterMetadata = Builder.Build(FShaderParametersMetadata::EUseCase::ShaderParameterStruct, TEXT("iree_add_parameters"));
+			}
+
+			FNNERuntimeIREEShader::FParameters *ShaderParameters = GraphBuilder.AllocParameters<FNNERuntimeIREEShader::FParameters>(ShaderParameterMetadata);
+			{
+				uint8* ShaderParameterDataPtr = (uint8*)ShaderParameters;
+				const TArray<FShaderParametersMetadata::FMember> &Members = ShaderParameterMetadata->GetMembers();
+				{
+					FRDGBufferUAVRef* ElementPtr = (FRDGBufferUAVRef*)(ShaderParameterDataPtr + Members[0].GetOffset());
+					*ElementPtr = GraphBuilder.CreateUAV(ACopy);
+				}
+				{
+					FRDGBufferUAVRef* ElementPtr = (FRDGBufferUAVRef*)(ShaderParameterDataPtr + Members[1].GetOffset());
+					*ElementPtr = GraphBuilder.CreateUAV(BCopy);
+				}
+				{
+					FRDGBufferUAVRef* ElementPtr = (FRDGBufferUAVRef*)(ShaderParameterDataPtr + Members[2].GetOffset());
+					*ElementPtr = GraphBuilder.CreateUAV(ResultCopy);
+				}
+			}
 
 			RDG_EVENT_SCOPE_STAT(GraphBuilder, FCommandBufferDispatch, "CommandBuffer.Dispatch");
 			RDG_GPU_STAT_SCOPE(GraphBuilder, FCommandBufferDispatch);
@@ -104,6 +130,7 @@ public:
 				RDG_EVENT_NAME("CommandBuffer.Dispatch"),
 				ERDGPassFlags::Compute | ERDGPassFlags::NeverCull,
 				Shader,
+				ShaderParameterMetadata,
 				ShaderParameters,
 				FIntVector(WorkgroupX, WorkgroupY, WorkgroupZ));
 
diff --git a/Engine/Plugins/Experimental/NNERuntimeIREE/Source/IREEDriverRDG/Private/IREEDriverRDGDevice.cpp b/Engine/Plugins/Experimental/NNERuntimeIREE/Source/IREEDriverRDG/Private/IREEDriverRDGDevice.cpp
--- a/Engine/Plugins/Experimental/NNERuntimeIREE/Source/IREEDriverRDG/Private/IREEDriverRDGDevice.cpp
+++ b/Engine/Plugins/Experimental/NNERuntimeIREE/Source/IREEDriverRDG/Private/IREEDriverRDGDevice.cpp
@@ -21,7 +21,7 @@ namespace Private
 class FDevice
 {
 public:
-	static iree_status_t Create(iree_allocator_t HostAllocator, iree_hal_device_t** OutDevice)
+	static iree_status_t Create(iree_allocator_t HostAllocator, const FString& BasePath, iree_hal_device_t** OutDevice)
 	{
 		check(OutDevice);
 		UE_LOG(LogIREEDriverRDG, Display, TEXT("%s"), StringCast<TCHAR>(__FUNCTION__).Get());
@@ -38,15 +38,12 @@ public:
 			return Status;
 		}
 
+		Device->BasePath = BasePath;
+
 		*OutDevice = (iree_hal_device_t*)Device;
 		return Status;
 	}
 
-	iree_allocator_t GetHostAllocator()
-	{
-		return HostAllocator;
-	}
-
 private:
 	static FDevice* Cast(iree_hal_device_t* Device)
 	{
@@ -167,7 +164,7 @@ private:
 		UE_LOG(LogIREEDriverRDG, Display, TEXT("%s"), StringCast<TCHAR>(__FUNCTION__).Get());
 
 		FDevice* Device = Cast(BaseDevice);
-		return NoOpExecutableCacheCreate(Device->HostAllocator, OutExecutableCache);
+		return NoOpExecutableCacheCreate(Device->HostAllocator, Device->GetBasePath(), OutExecutableCache);
 	}
 
 	static iree_status_t ImportFile(iree_hal_device_t* BaseDevice,
@@ -307,12 +304,23 @@ private:
 		return iree_make_status(IREE_STATUS_UNIMPLEMENTED, __FUNCTION__);
 	}
 
+	iree_allocator_t GetHostAllocator() const
+	{
+		return HostAllocator;
+	}
+
+	FString GetBasePath() const
+	{
+		return BasePath;
+	}
+
 	static const iree_hal_device_vtable_t VTable;
 
 	iree_hal_resource_t Resource;
 	iree_allocator_t HostAllocator;
 	iree_hal_allocator_t* DeviceAllocator;
 	iree_hal_channel_provider_t* ChannelProvider;
+	FString BasePath;
 };
 
 const iree_hal_device_vtable_t FDevice::VTable =
@@ -348,9 +356,9 @@ const iree_hal_device_vtable_t FDevice::VTable =
 
 } // namespace Private
 
-iree_status_t DeviceCreate(iree_allocator_t HostAllocator, iree_hal_device_t** OutDevice)
+iree_status_t DeviceCreate(iree_allocator_t HostAllocator, const FString& BasePath, iree_hal_device_t** OutDevice)
 {
-	return Private::FDevice::Create(HostAllocator, OutDevice);
+	return Private::FDevice::Create(HostAllocator, BasePath, OutDevice);
 }
 
 } // UE::IREE
diff --git a/Engine/Plugins/Experimental/NNERuntimeIREE/Source/IREEDriverRDG/Private/IREEDriverRDGExecutable.cpp b/Engine/Plugins/Experimental/NNERuntimeIREE/Source/IREEDriverRDG/Private/IREEDriverRDGExecutable.cpp
--- a/Engine/Plugins/Experimental/NNERuntimeIREE/Source/IREEDriverRDG/Private/IREEDriverRDGExecutable.cpp
+++ b/Engine/Plugins/Experimental/NNERuntimeIREE/Source/IREEDriverRDG/Private/IREEDriverRDGExecutable.cpp
@@ -5,6 +5,9 @@
 #ifdef WITH_IREE_DRIVER_RDG
 
 #include "IREEDriverRDGLog.h"
+#include "NNERuntimeIREEShaderShared.h"
+#include "Misc/FileHelper.h"
+#include "Serialization/MemoryReader.h"
 
 namespace UE::IREE::HAL::RDG
 {
@@ -15,7 +18,7 @@ namespace Private
 class FExecutable
 {
 public:
-	static iree_status_t Create(iree_allocator_t HostAllocator, iree_hal_executable_t** OutExecutable)
+	static iree_status_t Create(iree_allocator_t HostAllocator, const FString& BasePath, const iree_hal_executable_params_t* ExecutableParams, iree_hal_executable_t** OutExecutable)
 	{
 		check(OutExecutable);
 		UE_LOG(LogIREEDriverRDG, Display, TEXT("%s"), StringCast<TCHAR>(__FUNCTION__).Get());
@@ -25,9 +28,31 @@ public:
 		iree_hal_resource_initialize((const void*)&FExecutable::VTable, &Executable->Resource);
 		Executable->HostAllocator = HostAllocator;
 
+		{
+			const FString ShaderMapFileNameNoExt = TEXT("module_add_dispatch_0_add_dispatch_0_generic_4_f32");
+			const FString ShaderMapFileName = ShaderMapFileNameNoExt + TEXT(".ireeshadermap");
+			const FString ShaderMapFullPath = FPaths::Combine(BasePath, ShaderMapFileName);
+
+			TArray/*64*/<uint8> Data;
+			FFileHelper::LoadFileToArray(Data, *ShaderMapFullPath);
+
+			FMemoryReaderView Reader(Data, /*bIsPersitent =*/ true);
+
+			TUniquePtr<FNNERuntimeIREEResource> KernelResource = MakeUnique<FNNERuntimeIREEResource>();
+			KernelResource->SerializeShaderMap(Reader);
+
+			Executable->KernelResource = MoveTemp(KernelResource);
+		}
+
 		*OutExecutable = (iree_hal_executable_t*)Executable;
 		return iree_ok_status();
 	}
+
+	const FNNERuntimeIREEResource& GetResource() const
+	{
+		return *KernelResource;
+	}
+
 private:
 	static FExecutable* Cast(iree_hal_executable_t* Executable)
 	{
@@ -56,13 +81,19 @@ private:
 private:
 	iree_hal_resource_t Resource;
 	iree_allocator_t HostAllocator;
+	TUniquePtr<FNNERuntimeIREEResource> KernelResource;
 };
 
 } // namespace Private
 
-iree_status_t ExecutableCreate(iree_allocator_t HostAllocator, iree_hal_executable_t** OutExecutable)
+iree_status_t ExecutableCreate(iree_allocator_t HostAllocator, const FString& BasePath, const iree_hal_executable_params_t* ExecutableParams, iree_hal_executable_t** OutExecutable)
+{
+	return Private::FExecutable::Create(HostAllocator, BasePath, ExecutableParams, OutExecutable);
+}
+
+const FNNERuntimeIREEResource& ExecutableGetResource(iree_hal_executable_t* Executable)
 {
-	return Private::FExecutable::Create(HostAllocator, OutExecutable);
+	return ((Private::FExecutable*)Executable)->GetResource();
 }
 
 } // UE::IREE
diff --git a/Engine/Plugins/Experimental/NNERuntimeIREE/Source/IREEDriverRDG/Private/IREEDriverRDGExecutable.h b/Engine/Plugins/Experimental/NNERuntimeIREE/Source/IREEDriverRDG/Private/IREEDriverRDGExecutable.h
--- a/Engine/Plugins/Experimental/NNERuntimeIREE/Source/IREEDriverRDG/Private/IREEDriverRDGExecutable.h
+++ b/Engine/Plugins/Experimental/NNERuntimeIREE/Source/IREEDriverRDG/Private/IREEDriverRDGExecutable.h
@@ -6,10 +6,14 @@
 
 #include "IREEDriverRDG.h"
 
+class FNNERuntimeIREEResource;
+
 namespace UE::IREE::HAL::RDG
 {
 
-iree_status_t ExecutableCreate(iree_allocator_t HostAllocator, iree_hal_executable_t** OutExecutable);
+iree_status_t ExecutableCreate(iree_allocator_t HostAllocator, const FString& BasePath, const iree_hal_executable_params_t* ExecutableParams, iree_hal_executable_t** OutExecutable);
+
+const FNNERuntimeIREEResource& ExecutableGetResource(iree_hal_executable_t* Executable);
 
 } // UE::IREE
 
diff --git a/Engine/Plugins/Experimental/NNERuntimeIREE/Source/IREEDriverRDG/Private/IREEDriverRDGNoOpExecutableCache.cpp b/Engine/Plugins/Experimental/NNERuntimeIREE/Source/IREEDriverRDG/Private/IREEDriverRDGNoOpExecutableCache.cpp
--- a/Engine/Plugins/Experimental/NNERuntimeIREE/Source/IREEDriverRDG/Private/IREEDriverRDGNoOpExecutableCache.cpp
+++ b/Engine/Plugins/Experimental/NNERuntimeIREE/Source/IREEDriverRDG/Private/IREEDriverRDGNoOpExecutableCache.cpp
@@ -16,7 +16,7 @@ namespace Private
 class FNoOpExecutableCache
 {
 public:
-	static iree_status_t Create(iree_allocator_t HostAllocator, iree_hal_executable_cache_t** OutExecutableCache)
+	static iree_status_t Create(iree_allocator_t HostAllocator, const FString& BasePath, iree_hal_executable_cache_t** OutExecutableCache)
 	{
 		check(OutExecutableCache);
 		UE_LOG(LogIREEDriverRDG, Display, TEXT("%s"), StringCast<TCHAR>(__FUNCTION__).Get());
@@ -25,10 +25,17 @@ public:
 		IREE_RETURN_IF_ERROR(iree_allocator_malloc(HostAllocator, sizeof(*ExecutableCache), (void**)&ExecutableCache));
 		iree_hal_resource_initialize((const void*)&FNoOpExecutableCache::VTable, &ExecutableCache->Resource);
 		ExecutableCache->HostAllocator = HostAllocator;
+		ExecutableCache->BasePath = BasePath;
 
 		*OutExecutableCache = (iree_hal_executable_cache_t*)ExecutableCache;
 		return iree_ok_status();
 	}
+
+	FString GetBasePath() const
+	{
+		return BasePath;
+	}
+
 private:
 	static FNoOpExecutableCache* Cast(iree_hal_executable_cache_t* ExecutableCache)
 	{
@@ -47,20 +54,36 @@ private:
 	static bool CanPrepareFormat(iree_hal_executable_cache_t* BaseExecutableCache, iree_hal_executable_caching_mode_t CachingMode, iree_string_view_t ExecutableFormat) 
 	{
 		UE_LOG(LogIREEDriverRDG, Display, TEXT("%s format: %*.hs"), StringCast<TCHAR>(__FUNCTION__).Get(), (int)ExecutableFormat.size, ExecutableFormat.data);
-		return true;
+		if (iree_string_view_equal(ExecutableFormat, iree_make_cstring_view("vulkan-spirv-fb")))
+		{
+			return true;
+		}
+		// else if (iree_string_view_equal(ExecutableFormat, iree_make_cstring_view("vulkan-spirv-fb-ptr")))
+		// {
+		// 	return iree_all_bits_set(executable_cache->logical_device->enabled_features(), IREE_HAL_VULKAN_FEATURE_ENABLE_BUFFER_DEVICE_ADDRESSES);
+		// }
+		return false;
 	}
 
 	static iree_status_t PrepareExecutable(iree_hal_executable_cache_t* BaseExecutableCache, const iree_hal_executable_params_t* ExecutableParams, iree_hal_executable_t** OutExecutable) 
 	{
 		UE_LOG(LogIREEDriverRDG, Display, TEXT("%s"), StringCast<TCHAR>(__FUNCTION__).Get());
 
-		return ExecutableCreate(Cast(BaseExecutableCache)->HostAllocator, OutExecutable);
+		if (!FNoOpExecutableCache::CanPrepareFormat(BaseExecutableCache, ExecutableParams->caching_mode, ExecutableParams->executable_format))
+		{
+			return iree_make_status(IREE_STATUS_NOT_FOUND, "No executable implementation registered for the given executable format '%.*s'", (int)ExecutableParams->executable_format.size, ExecutableParams->executable_format.data);
+		}
+
+		FNoOpExecutableCache* ExecutableCache = Cast(BaseExecutableCache);
+
+		return ExecutableCreate(ExecutableCache->HostAllocator, ExecutableCache->GetBasePath(), ExecutableParams, OutExecutable);
 	}
 
 	static const iree_hal_executable_cache_vtable_t VTable;
 
 	iree_hal_resource_t Resource;
 	iree_allocator_t HostAllocator;
+	FString BasePath;
 };
 
 const iree_hal_executable_cache_vtable_t FNoOpExecutableCache::VTable = 
@@ -72,9 +95,9 @@ const iree_hal_executable_cache_vtable_t FNoOpExecutableCache::VTable =
 
 } // namespace Private
 
-iree_status_t NoOpExecutableCacheCreate(iree_allocator_t HostAllocator, iree_hal_executable_cache_t** OutExecutableCache)
+iree_status_t NoOpExecutableCacheCreate(iree_allocator_t HostAllocator, const FString& BasePath, iree_hal_executable_cache_t** OutExecutableCache)
 {
-	return Private::FNoOpExecutableCache::Create(HostAllocator, OutExecutableCache);
+	return Private::FNoOpExecutableCache::Create(HostAllocator, BasePath, OutExecutableCache);
 }
 
 } // UE::IREE
diff --git a/Engine/Plugins/Experimental/NNERuntimeIREE/Source/IREEDriverRDG/Private/IREEDriverRDGNoOpExecutableCache.h b/Engine/Plugins/Experimental/NNERuntimeIREE/Source/IREEDriverRDG/Private/IREEDriverRDGNoOpExecutableCache.h
--- a/Engine/Plugins/Experimental/NNERuntimeIREE/Source/IREEDriverRDG/Private/IREEDriverRDGNoOpExecutableCache.h
+++ b/Engine/Plugins/Experimental/NNERuntimeIREE/Source/IREEDriverRDG/Private/IREEDriverRDGNoOpExecutableCache.h
@@ -9,7 +9,7 @@
 namespace UE::IREE::HAL::RDG
 {
 
-iree_status_t NoOpExecutableCacheCreate(iree_allocator_t HostAllocator, iree_hal_executable_cache_t** OutExecutableCache);
+iree_status_t NoOpExecutableCacheCreate(iree_allocator_t HostAllocator, const FString& BasePath, iree_hal_executable_cache_t** OutExecutableCache);
 
 } // UE::IREE
 
diff --git a/Engine/Plugins/Experimental/NNERuntimeIREE/Source/IREEDriverRDGShaders/IREEDriverRDGShaders.Build.cs b/Engine/Plugins/Experimental/NNERuntimeIREE/Source/IREEDriverRDGShaders/IREEDriverRDGShaders.Build.cs
deleted file mode 100644
diff --git a/Engine/Plugins/Experimental/NNERuntimeIREE/Source/IREEDriverRDGShaders/Internal/IREEDriverRDGShadersAddCS.h b/Engine/Plugins/Experimental/NNERuntimeIREE/Source/IREEDriverRDGShaders/Internal/IREEDriverRDGShadersAddCS.h
deleted file mode 100644
diff --git a/Engine/Plugins/Experimental/NNERuntimeIREE/Source/IREEDriverRDGShaders/Internal/IREEDriverRDGShadersAddCrossCS.h b/Engine/Plugins/Experimental/NNERuntimeIREE/Source/IREEDriverRDGShaders/Internal/IREEDriverRDGShadersAddCrossCS.h
deleted file mode 100644
diff --git a/Engine/Plugins/Experimental/NNERuntimeIREE/Source/IREEDriverRDGShaders/Private/IREEDriverRDGShadersAddCS.cpp b/Engine/Plugins/Experimental/NNERuntimeIREE/Source/IREEDriverRDGShaders/Private/IREEDriverRDGShadersAddCS.cpp
deleted file mode 100644
diff --git a/Engine/Plugins/Experimental/NNERuntimeIREE/Source/IREEDriverRDGShaders/Private/IREEDriverRDGShadersAddCrossCS.cpp b/Engine/Plugins/Experimental/NNERuntimeIREE/Source/IREEDriverRDGShaders/Private/IREEDriverRDGShadersAddCrossCS.cpp
deleted file mode 100644
diff --git a/Engine/Plugins/Experimental/NNERuntimeIREE/Source/IREEDriverRDGShaders/Private/IREEDriverRDGShadersLog.h b/Engine/Plugins/Experimental/NNERuntimeIREE/Source/IREEDriverRDGShaders/Private/IREEDriverRDGShadersLog.h
deleted file mode 100644
diff --git a/Engine/Plugins/Experimental/NNERuntimeIREE/Source/IREEDriverRDGShaders/Private/IREEDriverRDGShadersModule.cpp b/Engine/Plugins/Experimental/NNERuntimeIREE/Source/IREEDriverRDGShaders/Private/IREEDriverRDGShadersModule.cpp
deleted file mode 100644
diff --git a/Engine/Plugins/Experimental/NNERuntimeIREE/Source/NNERuntimeIREE/Private/NNERuntimeIREE.cpp b/Engine/Plugins/Experimental/NNERuntimeIREE/Source/NNERuntimeIREE/Private/NNERuntimeIREE.cpp
--- a/Engine/Plugins/Experimental/NNERuntimeIREE/Source/NNERuntimeIREE/Private/NNERuntimeIREE.cpp
+++ b/Engine/Plugins/Experimental/NNERuntimeIREE/Source/NNERuntimeIREE/Private/NNERuntimeIREE.cpp
@@ -13,6 +13,7 @@
 
 #include "HAL/Platform.h"
 #include "Interfaces/ITargetPlatform.h"
+#include "IREECompilerRDG.h"
 #include "Kismet/GameplayStatics.h"
 #include "Misc/App.h"
 #include "Misc/Paths.h"
@@ -543,7 +544,7 @@ TSharedPtr<UE::NNE::FSharedModelData> UNNERuntimeIREERdg::CreateModelData(const
 
 	TArray64<uint8> ResultData;
 	TWeakObjectPtr<UNNERuntimeIREEModelData> IREEModelData = NewObject<UNNERuntimeIREEModelData>();
-	FNNERuntimeIREECompilerResultCPU CompilerResult{};
+	FIREECompileRDGCompilerResultCPU CompilerResult{};
 
 	bool bNeedCompileMlir = true;
 	if (PlatformFile.FileExists(*IREEModelDataFilePath))
@@ -561,19 +562,19 @@ TSharedPtr<UE::NNE::FSharedModelData> UNNERuntimeIREERdg::CreateModelData(const
 
 		{
 			FMemoryReaderView Reader(IREEModelData->CompilerResult, /*bIsPersitent =*/ true);
-			FNNERuntimeIREECompilerResultCPU::StaticStruct()->SerializeBin(Reader, &CompilerResult);
+			FIREECompileRDGCompilerResultCPU::StaticStruct()->SerializeBin(Reader, &CompilerResult);
 		}
 
 		bNeedCompileMlir = false;
 		for (int32 i = 0; i < CompilerResult.ArchitectureInfos.Num(); i++)
 		{
-			const FNNERuntimeIREEArchitectureInfoCPU& Info = CompilerResult.ArchitectureInfos[i];
+			const FIREECompileRDGArchitectureInfoCPU& Info = CompilerResult.ArchitectureInfos[i];
 			const FString SharedLibrarySubDirFullPath = FPaths::Combine(SharedLibraryDirFullPath, Info.RelativeDirPath);
 
-			// const FString SharedLibraryFilePath = FPaths::Combine(SharedLibrarySubDirFullPath, Info.SharedLibraryFileName);
+			const FString SharedLibraryFilePath = FPaths::Combine(SharedLibrarySubDirFullPath, Info.SharedLibraryFileName);
 			const FString VmfbFilePath = FPaths::Combine(SharedLibrarySubDirFullPath, Info.VmfbFileName);
 
-			// bNeedCompileMlir |= !PlatformFile.FileExists(*SharedLibraryFilePath);
+			bNeedCompileMlir |= !PlatformFile.FileExists(*SharedLibraryFilePath);
 			bNeedCompileMlir |= !PlatformFile.FileExists(*VmfbFilePath);
 		}
 	}
@@ -585,16 +586,14 @@ TSharedPtr<UE::NNE::FSharedModelData> UNNERuntimeIREERdg::CreateModelData(const
 		PlatformFile.DeleteDirectoryRecursively(*IntermediateDirFullPath);
 		PlatformFile.CreateDirectoryTree(*IntermediateDirFullPath);
 
-		TUniquePtr<UE::NNERuntimeIREE::CPU::FCompiler> Compiler = UE::NNERuntimeIREE::CPU::FCompiler::Make(TargetPlatformName);
+		TUniquePtr<UE::IREE::Compiler::RDG::FCompiler> Compiler = UE::IREE::Compiler::RDG::FCompiler::Make(TargetPlatform);
 		if (!Compiler.IsValid())
 		{
 			UE_LOG(LogNNERuntimeIREE, Warning, TEXT("UNNERuntimeIREERdg failed to create a compiler to compile for platform %s"), *TargetPlatformName);
 			return TSharedPtr<UE::NNE::FSharedModelData>();
 		}
 
-		TWeakObjectPtr<UNNERuntimeIREEModuleMetaData> CompilerModuleMetaData = NewObject<UNNERuntimeIREEModuleMetaData>();
-
-		if (!Compiler->CompileMlirRDG(FileData, FileIdString, IntermediateDirFullPath, CompilerResult, *CompilerModuleMetaData))
+		if (!Compiler->CompileMlir(FileData, FileIdString, IntermediateDirFullPath, CompilerResult))
 		{
 			UE_LOG(LogNNERuntimeIREE, Warning, TEXT("UNNERuntimeIREERdg failed to compile model %s"), *FileIdString);
 			return TSharedPtr<UE::NNE::FSharedModelData>();
@@ -609,12 +608,18 @@ TSharedPtr<UE::NNE::FSharedModelData> UNNERuntimeIREERdg::CreateModelData(const
 		}
 		if (IREEModelData->ModuleMetaData.IsEmpty())
 		{
+			TWeakObjectPtr<UNNERuntimeIREEModuleMetaData> CompilerModuleMetaData = NewObject<UNNERuntimeIREEModuleMetaData>();
+
+			FString FileDataString = "";
+			FileDataString.AppendChars((char*)FileData.GetData(), FileData.Num());
+			CompilerModuleMetaData->ParseFromString(FileDataString);
+
 			FMemoryWriter64 Writer(IREEModelData->ModuleMetaData, /*bIsPersitent =*/ true);
 			CompilerModuleMetaData->Serialize(Writer);
 		}
 		{
 			FMemoryWriter64 Writer(IREEModelData->CompilerResult, /*bIsPersitent =*/ true);
-			FNNERuntimeIREECompilerResultCPU::StaticStruct()->SerializeBin(Writer, &CompilerResult);
+			FIREECompileRDGCompilerResultCPU::StaticStruct()->SerializeBin(Writer, &CompilerResult);
 		}
 
 		{
@@ -631,17 +636,17 @@ TSharedPtr<UE::NNE::FSharedModelData> UNNERuntimeIREERdg::CreateModelData(const
 	{
 		SCOPED_NAMED_EVENT_TEXT("Copy", FColor::Magenta);
 
-		const FNNERuntimeIREEArchitectureInfoCPU& Info = CompilerResult.ArchitectureInfos[i];
+		const FIREECompileRDGArchitectureInfoCPU& Info = CompilerResult.ArchitectureInfos[i];
 		const FString SharedLibrarySubDirFullPath = FPaths::Combine(SharedLibraryDirFullPath, Info.RelativeDirPath);
 		const FString StagingSubDirFullPath = FPaths::Combine(StagingDirFullPath, Info.Architecture);
 		
-		// const FString SharedLibraryFilePathSrc = FPaths::Combine(SharedLibrarySubDirFullPath, Info.SharedLibraryFileName);
+		const FString SharedLibraryFilePathSrc = FPaths::Combine(SharedLibrarySubDirFullPath, Info.SharedLibraryFileName);
 		const FString VmfbFilePathSrc = FPaths::Combine(SharedLibrarySubDirFullPath, Info.VmfbFileName);
 
-		// const FString SharedLibraryFilePathDest = FPaths::Combine(StagingSubDirFullPath, Info.SharedLibraryFileName);
+		const FString SharedLibraryFilePathDest = FPaths::Combine(StagingSubDirFullPath, Info.SharedLibraryFileName);
 		const FString VmfbFilePathDest = FPaths::Combine(StagingSubDirFullPath, Info.VmfbFileName);
 
-		// IFileManager::Get().Copy(*SharedLibraryFilePathDest, *SharedLibraryFilePathSrc, bNeedCompileMlir);
+		IFileManager::Get().Copy(*SharedLibraryFilePathDest, *SharedLibraryFilePathSrc, bNeedCompileMlir);
 		IFileManager::Get().Copy(*VmfbFilePathDest, *VmfbFilePathSrc, bNeedCompileMlir);
 	}
 
@@ -729,10 +734,10 @@ TSharedPtr<UE::NNE::IModelRDG> UNNERuntimeIREERdg::CreateModelRDG(const TObjectP
 		return {};
 	}
 
-	FNNERuntimeIREECompilerResultCPU CompilerResult{};
+	FIREECompileRDGCompilerResultCPU CompilerResult{};
 	{
 		FMemoryReaderView Reader(IREEModelData->CompilerResult, /*bIsPersitent =*/ true);
-		FNNERuntimeIREECompilerResultCPU::StaticStruct()->SerializeBin(Reader, &CompilerResult);
+		FIREECompileRDGCompilerResultCPU::StaticStruct()->SerializeBin(Reader, &CompilerResult);
 	}
 
 	int32 ArchitectureIndex = -1;
@@ -750,7 +755,7 @@ TSharedPtr<UE::NNE::IModelRDG> UNNERuntimeIREERdg::CreateModelRDG(const TObjectP
 		return {};
 	}
 
-	const FNNERuntimeIREEArchitectureInfoCPU& ArchitectureInfo = CompilerResult.ArchitectureInfos[ArchitectureIndex];
+	const FIREECompileRDGArchitectureInfoCPU& ArchitectureInfo = CompilerResult.ArchitectureInfos[ArchitectureIndex];
 
 	const FString FileIdString = IREEModelData->FileId.ToString(EGuidFormats::Digits).ToLower();
 	const FString SharedLibraryDirFullPath = FPaths::ConvertRelativePathToFull(FPaths::Combine(FPaths::ProjectDir(), GetSharedLibDirPath(UGameplayStatics::GetPlatformName(), GetRuntimeName(), FileIdString)));
diff --git a/Engine/Plugins/Experimental/NNERuntimeIREE/Source/NNERuntimeIREE/Private/NNERuntimeIREECompiler.cpp b/Engine/Plugins/Experimental/NNERuntimeIREE/Source/NNERuntimeIREE/Private/NNERuntimeIREECompiler.cpp
--- a/Engine/Plugins/Experimental/NNERuntimeIREE/Source/NNERuntimeIREE/Private/NNERuntimeIREECompiler.cpp
+++ b/Engine/Plugins/Experimental/NNERuntimeIREE/Source/NNERuntimeIREE/Private/NNERuntimeIREECompiler.cpp
@@ -352,67 +352,6 @@ namespace UE::NNERuntimeIREE
 
 			return bResult;
 		}
-		
-		bool FCompiler::CompileMlirRDG(TConstArrayView<uint8> InFileData, const FString& InModelName, const FString& InOutputDir, FNNERuntimeIREECompilerResultCPU& OutCompilerResult, UNNERuntimeIREEModuleMetaData& ModuleMetaData)
-		{
-			using namespace Private;
-
-			{
-				SCOPED_NAMED_EVENT_TEXT("Metadata", FColor::Magenta);
-
-				FString FileDataString = "";
-				FileDataString.AppendChars((char*)InFileData.GetData(), InFileData.Num());
-				ModuleMetaData.ParseFromString(FileDataString);
-			}
-
-			IPlatformFile& PlatformFile = FPlatformFileManager::Get().GetPlatformFile();
-
-			FString InputFilePath = FPaths::Combine(InOutputDir, InModelName) + ".mlir";
-			if (!PlatformFile.FileExists(*InputFilePath))
-			{
-				SCOPED_NAMED_EVENT_TEXT("InputFile", FColor::Magenta);
-
-				FFileHelper::SaveArrayToFile(InFileData, *InputFilePath);
-			}
-
-			bool bResult = true;
-			for (int32 i = 0; i < BuildTargets.Num(); i++)
-			{
-				FString IntermediateDirPath = FPaths::Combine(InOutputDir, BuildTargets[i].Architecture);
-				PlatformFile.CreateDirectoryTree(*IntermediateDirPath);
-				FString IntermediateFilePathNoExt = FPaths::Combine(IntermediateDirPath, InModelName);
-				FString VmfbFilePath = IntermediateFilePathNoExt + ".vmfb";
-
-				FString CompilerArguments = TEXT("--iree-hal-target-backends=vulkan-spirv -o ${VMFB_PATH} ${INPUT_PATH}");
-				CompilerArguments.ReplaceInline(*FString("${VMFB_PATH}"), *(FString("\"") + VmfbFilePath + "\""));
-				CompilerArguments.ReplaceInline(*FString("${INPUT_PATH}"), *(FString("\"") + InputFilePath + "\""));
-
-				RunCommand(CompilerCommand, CompilerArguments, IntermediateFilePathNoExt + "_compile-log.txt");
-
-				if (!PlatformFile.FileExists(*VmfbFilePath))
-				{
-					UE_LOG(LogNNERuntimeIREE, Warning, TEXT("UNNERuntimeIREERdg failed to compile the model \"%s\" using the command:"), *InputFilePath);
-					UE_LOG(LogNNERuntimeIREE, Warning, TEXT("\"%s\" %s"), *CompilerCommand, *CompilerArguments);
-					bResult = false;
-					continue;
-				}
-
-				FNNERuntimeIREEArchitectureInfoCPU ArchitectureInfo;
-				ArchitectureInfo.Architecture = BuildTargets[i].Architecture;
-				ArchitectureInfo.RelativeDirPath = BuildTargets[i].Architecture;
-				ArchitectureInfo.VmfbFileName = InModelName + ".vmfb";
-				OutCompilerResult.ArchitectureInfos.Add(MoveTemp(ArchitectureInfo));
-			}
-
-			bResult &= !OutCompilerResult.ArchitectureInfos.IsEmpty();
-			
-			if (!bResult)
-			{
-				OutCompilerResult.ArchitectureInfos.Empty();
-			}
-			
-			return bResult;
-		}
 	} // CPU
 } // UE::NNERuntimeIREE
 
diff --git a/Engine/Plugins/Experimental/NNERuntimeIREE/Source/NNERuntimeIREE/Private/NNERuntimeIREECompiler.h b/Engine/Plugins/Experimental/NNERuntimeIREE/Source/NNERuntimeIREE/Private/NNERuntimeIREECompiler.h
--- a/Engine/Plugins/Experimental/NNERuntimeIREE/Source/NNERuntimeIREE/Private/NNERuntimeIREECompiler.h
+++ b/Engine/Plugins/Experimental/NNERuntimeIREE/Source/NNERuntimeIREE/Private/NNERuntimeIREECompiler.h
@@ -82,7 +82,6 @@ namespace UE::NNERuntimeIREE::CPU
 
 	public:
 		bool CompileMlir(TConstArrayView<uint8> InFileData, const FString& InModelName, const FString& InOutputDir, FNNERuntimeIREECompilerResultCPU& OutCompilerResult, UNNERuntimeIREEModuleMetaData& ModuleMetaData);
-		bool CompileMlirRDG(TConstArrayView<uint8> InFileData, const FString& InModelName, const FString& InOutputDir, FNNERuntimeIREECompilerResultCPU& OutCompilerResult, UNNERuntimeIREEModuleMetaData& ModuleMetaData);
 
 	private:
 		FString CompilerCommand;
diff --git a/Engine/Plugins/Experimental/NNERuntimeIREE/Source/NNERuntimeIREE/Private/NNERuntimeIREEModel.cpp b/Engine/Plugins/Experimental/NNERuntimeIREE/Source/NNERuntimeIREE/Private/NNERuntimeIREEModel.cpp
--- a/Engine/Plugins/Experimental/NNERuntimeIREE/Source/NNERuntimeIREE/Private/NNERuntimeIREEModel.cpp
+++ b/Engine/Plugins/Experimental/NNERuntimeIREE/Source/NNERuntimeIREE/Private/NNERuntimeIREEModel.cpp
@@ -1019,7 +1019,7 @@ namespace UE::NNERuntimeIREE
 					iree_hal_device_release(Device);
 				}
 
-				static TSharedPtr<FDevice> Make()
+				static TSharedPtr<FDevice> Make(const FString& BasePath)
 				{
 					TSharedPtr<UE::NNERuntimeIREE::Private::FInstance> Instance = UE::NNERuntimeIREE::Private::FInstance::GetInstance();
 					if (!Instance.IsValid())
@@ -1028,7 +1028,7 @@ namespace UE::NNERuntimeIREE
 					}
 
 					iree_hal_device_t* Device = nullptr;
-					iree_status_t Status = IREE::HAL::RDG::DeviceCreate(iree_allocator_system(), &Device);
+					iree_status_t Status = IREE::HAL::RDG::DeviceCreate(iree_allocator_system(), BasePath, &Device);
 					if (!iree_status_is_ok(Status))
 					{
 						UE_LOG(LogNNERuntimeIREE, Error, TEXT("Could not create IREE RDG device!"));
@@ -1415,7 +1415,7 @@ namespace UE::NNERuntimeIREE
 			check(!InVmfbFileName.IsEmpty());
 			check(!InModuleMetaData.FunctionMetaData.IsEmpty());
 
-			TSharedPtr<Private::FDevice> Device = Private::FDevice::Make();
+			TSharedPtr<Private::FDevice> Device = Private::FDevice::Make(InDirPath);
 			if (!Device.IsValid())
 			{
 				return {};
diff --git a/Engine/Plugins/Experimental/NNERuntimeIREE/Source/NNERuntimeIREEShader/Internal/NNERuntimeIREEShader.h b/Engine/Plugins/Experimental/NNERuntimeIREE/Source/NNERuntimeIREEShader/Internal/NNERuntimeIREEShader.h
new file mode 100644
--- /dev/null
+++ b/Engine/Plugins/Experimental/NNERuntimeIREE/Source/NNERuntimeIREEShader/Internal/NNERuntimeIREEShader.h
@@ -0,0 +1,21 @@
+// Copyright Epic Games, Inc. All Rights Reserved.
+
+#pragma once
+
+#include "Shader.h"
+#include "ShaderParameters.h"
+#include "NNERuntimeIREEShaderType.h"
+#include "NNERuntimeIREEShaderShared.h"
+
+class NNERUNTIMEIREESHADER_API FNNERuntimeIREEShader : public FShader
+{
+	DECLARE_SHADER_TYPE(FNNERuntimeIREEShader, NNERuntimeIREE);
+
+public:
+	FNNERuntimeIREEShader() = default;
+	FNNERuntimeIREEShader(const FNNERuntimeIREEShaderType::CompiledShaderInitializerType & Initializer);
+
+	MS_ALIGN(SHADER_PARAMETER_STRUCT_ALIGNMENT) struct FParameters
+	{
+	};
+};
\ No newline at end of file
diff --git a/Engine/Plugins/Experimental/NNERuntimeIREE/Source/NNERuntimeIREEShader/Internal/NNERuntimeIREEShaderCompileResult.h b/Engine/Plugins/Experimental/NNERuntimeIREE/Source/NNERuntimeIREEShader/Internal/NNERuntimeIREEShaderCompileResult.h
new file mode 100644
--- /dev/null
+++ b/Engine/Plugins/Experimental/NNERuntimeIREE/Source/NNERuntimeIREEShader/Internal/NNERuntimeIREEShaderCompileResult.h
@@ -0,0 +1,37 @@
+// Copyright Epic Games, Inc. All Rights Reserved.
+
+#pragma once
+
+#include "CoreMinimal.h"
+
+/** Struct containing a message returned by shader compilation. */
+struct NNERUNTIMEIREESHADER_API FNNERuntimeIREEShaderCompileMessage
+{
+	enum class EMessageType
+	{
+		None,
+		Info,
+		Warning,
+		Error
+	};
+
+	EMessageType Type = EMessageType::None;
+	FString Text;
+	FString VirtualFilePath;
+	FString RealFilePath;
+	int32 Line = -1;
+	int32 ColumnStart = -1;
+	int32 ColumnEnd = -1;
+	
+	/** Equality operator used to deduplicate messages. */
+	bool operator==(FNNERuntimeIREEShaderCompileMessage const& Rhs) const
+	{
+		return Type == Rhs.Type && Text == Rhs.Text && VirtualFilePath == Rhs.VirtualFilePath && RealFilePath == Rhs.RealFilePath && Line == Rhs.Line && ColumnStart == Rhs.ColumnStart && ColumnEnd == Rhs.ColumnEnd;
+	}
+};
+
+/** Struct containing all of the messages returned by shader compilation. */
+struct NNERUNTIMEIREESHADER_API FNNERuntimeIREEShaderCompileResults
+{
+	TArray<FNNERuntimeIREEShaderCompileMessage> Messages;
+};
diff --git a/Engine/Plugins/Experimental/NNERuntimeIREE/Source/NNERuntimeIREEShader/Internal/NNERuntimeIREEShaderLog.h b/Engine/Plugins/Experimental/NNERuntimeIREE/Source/NNERuntimeIREEShader/Internal/NNERuntimeIREEShaderLog.h
new file mode 100644
--- /dev/null
+++ b/Engine/Plugins/Experimental/NNERuntimeIREE/Source/NNERuntimeIREEShader/Internal/NNERuntimeIREEShaderLog.h
@@ -0,0 +1,7 @@
+// Copyright Epic Games, Inc. All Rights Reserved.
+
+#pragma once
+
+#include "Logging/LogMacros.h"
+
+DECLARE_LOG_CATEGORY_EXTERN(LogNNERuntimeIREEShader, Log, All);
\ No newline at end of file
diff --git a/Engine/Plugins/Experimental/NNERuntimeIREE/Source/NNERuntimeIREEShader/Internal/NNERuntimeIREEShaderShared.h b/Engine/Plugins/Experimental/NNERuntimeIREE/Source/NNERuntimeIREEShader/Internal/NNERuntimeIREEShaderShared.h
new file mode 100644
--- /dev/null
+++ b/Engine/Plugins/Experimental/NNERuntimeIREE/Source/NNERuntimeIREEShader/Internal/NNERuntimeIREEShaderShared.h
@@ -0,0 +1,426 @@
+// Copyright Epic Games, Inc. All Rights Reserved.
+
+#pragma once
+
+#include "CoreMinimal.h"
+#include "NNERuntimeIREEShaderCompileResult.h"
+#include "Engine/EngineTypes.h"
+#include "RenderResource.h"
+#include "RenderingThread.h"
+#include "RenderDeferredCleanup.h"
+#include "RHI.h"
+#include "SceneTypes.h"
+#include "Shader.h"
+#include "ShaderCompiler.h"
+#include "Templates/RefCounting.h"
+
+class FNNERuntimeIREEResource;
+class FNNERuntimeIREEShader;
+
+/** Stores outputs from the  kernel compile that need to be saved. */
+class FNNERuntimeIREECompilationOutput
+{
+	DECLARE_TYPE_LAYOUT(FNNERuntimeIREECompilationOutput, NonVirtual);
+public:
+	FNNERuntimeIREECompilationOutput()
+	{}
+};
+
+/** Contains all the information needed to uniquely identify a FNNERuntimeIREEShaderMapID. */
+class FNNERuntimeIREEShaderMapId
+{
+	DECLARE_TYPE_LAYOUT(FNNERuntimeIREEShaderMapId, NonVirtual);
+public:
+	/** Feature level that the shader map is going to be compiled for.  */
+	LAYOUT_FIELD(ERHIFeatureLevel::Type, FeatureLevel);
+
+	/**
+	 * The hash computed from the shader code generated by the library.
+	 */
+	LAYOUT_FIELD(uint64, ShaderCodeHash);
+
+	/** Shader types of shaders that are inlined in this shader map in the DDC. */
+	LAYOUT_FIELD(TMemoryImageArray<FShaderTypeDependency>, ShaderTypeDependencies);
+
+	/*
+	* Type layout parameters of the memory image
+	*/
+	LAYOUT_FIELD(FPlatformTypeLayoutParameters, LayoutParams);
+
+public:
+	FNNERuntimeIREEShaderMapId()
+		: FeatureLevel(ERHIFeatureLevel::SM5)
+	{ }
+
+	~FNNERuntimeIREEShaderMapId() = default;
+
+#if WITH_EDITOR
+	void SetShaderDependencies(const TArray<FShaderType*>& InShaderTypes, EShaderPlatform InShaderPlatform);
+#endif
+
+	//void Serialize(FArchive& Ar);
+
+	friend uint32 GetTypeHash(const FNNERuntimeIREEShaderMapId& Ref)
+	{
+		return GetTypeHash(Ref.ShaderCodeHash);
+	}
+
+	SIZE_T GetSizeBytes() const
+	{
+		return sizeof(*this);
+	}
+
+	/** Hashes the kernel specific part of this shader map Id. */
+	void GetKernelHash(FSHAHash& OutHash) const;
+
+	/**
+	 * Tests this set against another for equality, disregarding override settings.
+	 *
+	 * @param InReferenceSet	The set to compare against
+	 * @return					true if the sets are equal
+	 */
+	bool operator==(const FNNERuntimeIREEShaderMapId& InReferenceSet) const;
+
+	bool operator!=(const FNNERuntimeIREEShaderMapId& InReferenceSet) const
+	{
+		return !(*this == InReferenceSet);
+	}
+
+#if WITH_EDITOR
+	/** Appends string representations of this Id to a key string. */
+	void AppendKeyString(FString& OutKeyString) const;
+#endif // WITH_EDITOR
+
+	/** Returns true if the requested shader type is a dependency of this shader map Id. */
+	bool ContainsShaderType(const FShaderType* ShaderType) const;
+};
+
+class FNNERuntimeIREEShaderMapContent : public FShaderMapContent
+{
+	using Super = FShaderMapContent;
+	friend class FNNERuntimeIREEShaderMap;
+	DECLARE_TYPE_LAYOUT(FNNERuntimeIREEShaderMapContent, NonVirtual);
+private:
+	explicit FNNERuntimeIREEShaderMapContent(EShaderPlatform InPlatform) : Super(InPlatform) {}
+
+	/** The kernel's user friendly name */
+	LAYOUT_FIELD(FMemoryImageString, FriendlyName);
+
+	/** The static parameter set that this shader map was compiled with */
+	LAYOUT_FIELD(FNNERuntimeIREEShaderMapId, ShaderMapId);
+
+	/** Shader compilation output */
+	LAYOUT_FIELD(FNNERuntimeIREECompilationOutput, CompilationOutput);
+};
+
+
+/**
+ * All the shader permutations for a kernel.
+ */
+class FNNERuntimeIREEShaderMap : public TShaderMap<FNNERuntimeIREEShaderMapContent, FShaderMapPointerTable>, public FDeferredCleanupInterface
+{
+public:
+	using Super = TShaderMap<FNNERuntimeIREEShaderMapContent, FShaderMapPointerTable>;
+
+	/**
+	 * Finds the shader map for a kernel.
+	 * @param InShaderMapId - Id to find
+	 * @param Platform - The platform to lookup for
+	 * @return nullptr if no cached shader map was found.
+	 */
+	static FNNERuntimeIREEShaderMap* FindId(const FNNERuntimeIREEShaderMapId& InShaderMapId, EShaderPlatform InPlatform);
+
+	// ShaderMap interface
+	template<typename ShaderType> TShaderRef<ShaderType> GetShader(int32 PermutationId) const { return TShaderRef<ShaderType>(GetContent()->GetShader<ShaderType>(PermutationId), *this); }
+	TShaderRef<FShader> GetShader(FShaderType* ShaderType, int32 PermutationId) const { return TShaderRef<FShader>(GetContent()->GetShader(ShaderType, PermutationId), *this); }
+
+	FNNERuntimeIREEShaderMap();
+	~FNNERuntimeIREEShaderMap();
+
+#if WITH_EDITOR
+	/**
+	 * Compiles the shaders for a kernel and caches them in this shader map.
+	 * @param InKernel - The kernel to compile shaders for.
+	 * @param InShaderMapId - the set of static parameters to compile for
+	 * @param InPlatform - The platform to compile to
+	 */
+	void Compile(
+		FNNERuntimeIREEResource* InKernel,
+		const FNNERuntimeIREEShaderMapId& InShaderMapId,
+		TRefCountPtr<FSharedShaderCompilerEnvironment> InCompilationEnvironment,
+		const FNNERuntimeIREECompilationOutput& InCompilationOutput,
+		EShaderPlatform InPlatform,
+		bool bSynchronousCompile,
+		bool bApplyCompletedShaderMapForRendering
+		);
+
+	/** Sorts the incoming compiled jobs into the appropriate shader maps, and finalizes this shader map so that it can be used for rendering. */
+	bool ProcessCompilationResults(const TArray<FShaderCommonCompileJobPtr>& InCompilationResults, int32& InOutResultIndex, float& InOutTimeBudget);
+#endif // WITH_EDITOR
+
+	/**
+	 * Checks whether the shader map is missing any shader types necessary for the given kernel.
+	 * @param InKernel - The ComputeKernel which is checked.
+	 * @return True if the shader map has all of the shader types necessary.
+	 */
+	bool IsComplete(const FNNERuntimeIREEResource* InKernel, bool bSilent);
+
+	/** Builds a list of the shaders in a shader map. */
+	void GetShaderList(TMap<FShaderId, TShaderRef<FShader>>& OutShaders) const;
+	virtual void GetShaderList(TMap<FHashedName, TShaderRef<FShader>>& OutShaders) const override;
+	virtual void GetShaderPipelineList(TArray<FShaderPipelineRef>& OutShaderPipelines) const override;
+
+	/** Registers a ComputeKernel shader map in the global map so it can be used. */
+	void Register(EShaderPlatform InShaderPlatform);
+
+	// Reference counting.
+	void AddRef();
+	void Release();
+
+	/** Serializes the shader map. */
+	bool Serialize(FArchive& Ar);
+
+	// Accessors
+	const FNNERuntimeIREEShaderMapId& GetShaderMapId() const { return GetContent()->ShaderMapId; }
+	EShaderPlatform GetShaderPlatform() const { return GetContent()->GetShaderPlatform(); }
+	const FMemoryImageString& GetFriendlyName() const { return GetContent()->FriendlyName; }
+	bool IsCompilationFinalized() const { return bCompilationFinalized; }
+	bool CompiledSuccessfully() const { return bCompiledSuccessfully; }
+
+#if WITH_EDITOR
+	static TMap<TRefCountPtr<FNNERuntimeIREEShaderMap>, TArray<FNNERuntimeIREEResource*> > &GetInFlightShaderMaps() { return ShaderMapsBeingCompiled; }
+
+	void SetCompiledSuccessfully(bool bSuccess) { bCompiledSuccessfully = bSuccess; }
+#endif // WITH_EDITOR
+
+	int32 GetNumRefs() const { return NumRefs; }
+	uint32 GetCompilingId()  { return CompilingId; }
+
+private:
+#if WITH_EDITOR
+	FShader* ProcessCompilationResultsForSingleJob(FShaderCompileJob& InSingleJob, const FSHAHash& InShaderMapHash);
+#endif
+
+	bool IsComputeKernelShaderComplete(const FNNERuntimeIREEResource* InKernel, const FNNERuntimeIREEShaderType* InShaderType, bool bSilent);
+
+	/**
+	 * A global map from a kernel's ID and static switch set to any shader map cached for that kernel.
+	 * Note: this does not necessarily contain all kernel shader maps in memory.  Shader maps with the same key can evict each other.
+	 * No ref counting needed as these are removed on destruction of the shader map.
+	 */
+	static TMap<FNNERuntimeIREEShaderMapId, FNNERuntimeIREEShaderMap*> GIdToKernelShaderMap[SP_NumPlatforms];
+
+	/**
+	 * All kernel shader maps in memory.
+	 * No ref counting needed as these are removed on destruction of the shader map.
+	 */
+	static TArray<FNNERuntimeIREEShaderMap*> AllKernelShaderMaps;
+
+#if WITH_EDITOR
+	/** Tracks resources and their shader maps that need to be compiled but whose compilation is being deferred. */
+	static TMap<TRefCountPtr<FNNERuntimeIREEShaderMap>, TArray<FNNERuntimeIREEResource*> > ShaderMapsBeingCompiled;
+#endif
+
+	/** Uniquely identifies this shader map during compilation, needed for deferred compilation where shaders from multiple shader maps are compiled together. */
+	uint32 CompilingId;
+
+	mutable int32 NumRefs;
+
+	/** Used to catch errors where the shader map is deleted directly. */
+	bool bDeletedThroughDeferredCleanup;
+
+	/** Indicates whether this shader map has been registered in GIdToComputeKernelShaderMap */
+	uint32 bRegistered : 1;
+
+	/**
+	 * Indicates whether this shader map has had ProcessCompilationResults called after Compile.
+	 * The shader map must not be used on the rendering thread unless bCompilationFinalized is true.
+	 */
+	uint32 bCompilationFinalized : 1;
+
+	uint32 bCompiledSuccessfully : 1;
+
+};
+
+class NNERUNTIMEIREESHADER_API FNNERuntimeIREEResource
+{
+public:
+	FNNERuntimeIREEResource();
+	virtual ~FNNERuntimeIREEResource();
+
+	/**
+	 * Caches the shaders for this kernel with no static parameters on the given platform.
+	 * This is used by UKernel
+	 */
+	bool CacheShaders(EShaderPlatform InPlatform, const ITargetPlatform* TargetPlatform, bool bApplyCompletedShaderMapForRendering, bool bSynchronous);
+	bool CacheShaders(const FNNERuntimeIREEShaderMapId& InShaderMapId, EShaderPlatform InPlatform, bool bApplyCompletedShaderMapForRendering, bool bSynchronous);
+
+	/**
+	 * Should the shader for this kernel with the given platform and shader type combination be compiled
+	 *
+	 * @param InPlatform		The platform currently being compiled for
+	 * @param InShaderType		Which shader is being compiled
+	 *
+	 * @return true if the shader should be compiled
+	 */
+	virtual bool ShouldCache(EShaderPlatform InPlatform, const FShaderType* InShaderType) const;
+
+	void SerializeShaderMap(FArchive& Ar);
+
+	void GetDependentShaderTypes(EShaderPlatform InPlatform, TArray<FShaderType*>& OutShaderTypes) const;
+	virtual void GetShaderMapId(EShaderPlatform InPlatform, const ITargetPlatform* TargetPlatform, FNNERuntimeIREEShaderMapId& OutId) const;
+
+	/**
+	 * Called when compilation finishes, after the GameThreadShaderMap is set and the render command to set the RenderThreadShaderMap is queued
+	 */
+	virtual void NotifyCompilationFinished(FString const& ResultMessage);
+
+#if WITH_EDITOR
+	/**
+	 * Cancels all outstanding compilation jobs
+	 */
+	void CancelCompilation();
+
+	/**
+	 * Blocks until compilation has completed. Returns immediately if a compilation is not outstanding.
+	 */
+	void FinishCompilation();
+#endif // WITH_EDITOR
+
+	// Accessors.
+	FNNERuntimeIREEShaderCompileResults const& GetCompilationResults() const { return CompilationResults; }
+	void SetCompilationResults(FNNERuntimeIREEShaderCompileResults const& InCompilationResults) { CompilationResults = InCompilationResults; }
+
+	ERHIFeatureLevel::Type GetFeatureLevel() const { return FeatureLevel; }
+
+	FNNERuntimeIREEShaderMap* GetGameThreadShaderMap() const
+	{
+		checkSlow(IsInGameThread() || IsInAsyncLoadingThread());
+		return GameThreadShaderMap;
+	}
+
+	/** Returns owner name for tracking */
+	FName GetOwnerFName() const
+	{
+		return AssetPath;
+	}
+
+	/** Note: SetRenderingThreadShaderMap must also be called with the same value, but from the rendering thread. */
+	void SetGameThreadShaderMap(FNNERuntimeIREEShaderMap* InShaderMap)
+	{
+		checkSlow(IsInGameThread() || IsInAsyncLoadingThread());
+		GameThreadShaderMap = InShaderMap;
+		if (LIKELY(GameThreadShaderMap))
+		{
+			GameThreadShaderMap->GetResource()->SetOwnerName(GetOwnerFName());
+		}
+	}
+
+	void AddCompileId(uint32 InIdentifier) 
+	{
+		OutstandingCompileShaderMapIds.Add(InIdentifier);
+	}
+
+	void RemoveOutstandingCompileId(const int32 InOldOutstandingCompileShaderMapId);
+
+	const FString& GetHLSLSource() const
+	{
+		return ShaderSource;
+	}
+
+	const FString& GetEntryPoint() const
+	{
+		return ShaderEntryPoint;
+	}
+
+	const FShaderParametersMetadata* GetShaderParamMetadata() const
+	{
+		return ShaderParameterMetadata;
+	}
+
+	const FString& GetFriendlyName() const { return FriendlyName; }
+
+	int32 GetNumPermutations() const { return 1; }
+
+	void SetupResource(
+		ERHIFeatureLevel::Type InFeatureLevel,
+		FString const& InFriendlyName,
+		FString const& InShaderEntryPoint,
+		FString const& InShaderHashKey,
+		FString const& InShaderSource,
+		// TMap<FString, FString> const& InAdditionalSources,
+		// TMap<FString, FString> const& InGeneratedSources,
+		// TSharedPtr<FComputeKernelDefinitionSet>& InShaderDefinitionSet,
+		// TSharedPtr<FComputeKernelPermutationVector>& InShaderPermutationVector,
+		// TUniquePtr<FShaderParametersMetadataAllocations>& InShaderParameterMetadataAllocations,
+		FShaderParametersMetadata* InShaderParameterMetadata,
+		FName const& InAssetPath
+	);
+
+	TShaderRef<FNNERuntimeIREEShader> GetShader(int32 PermutationId) const;
+
+	bool IsSame(const FNNERuntimeIREEShaderMapId& InId) const;
+
+protected:
+#if WITH_EDITOR
+	/**
+	 * Fills the passed array with IDs of shader maps unfinished compilation jobs.
+	 */
+	void GetShaderMapIDsWithUnfinishedCompilation(TArray<int32>& OutShaderMapIds);
+#endif // WITH_EDITOR
+
+private:
+	FNNERuntimeIREEShaderCompileResults CompilationResults;
+
+	/** 
+	 * Game thread tracked shader map, which is ref counted and manages shader map lifetime. 
+	 * The shader map uses deferred deletion so that the rendering thread has a chance to process a release command when the shader map is no longer referenced.
+	 * Code that sets this is responsible for updating RenderingThreadShaderMap in a thread safe way.
+	 * During an async compile, this will be NULL and will not contain the actual shader map until compilation is complete.
+	 */
+	TRefCountPtr<FNNERuntimeIREEShaderMap> GameThreadShaderMap = nullptr;
+
+	/** Name of shader main function. */
+	FString ShaderEntryPoint;
+
+	/** Cached shader source for compilation. */
+	FString ShaderSource;
+
+	/** 
+	 * Hash computed from shader code. If multiple kernels give
+	 * the same output, same shader resource will be shared 
+	 */
+	uint64 ShaderCodeHash = 0;
+
+	/** Shader parameter metadata. It is expected that object is owned by ShaderParameterMetadataAllocations.  */
+	FShaderParametersMetadata const* ShaderParameterMetadata;
+
+#if WITH_EDITOR
+	/**
+	 * Compiles this kernel for InPlatform, storing the result in OutShaderMap if the compile was synchronous
+	 */
+	bool BeginCompileShaderMap(
+		const FNNERuntimeIREEShaderMapId& InShaderMapId,
+		EShaderPlatform InPlatform, 
+		TRefCountPtr<FNNERuntimeIREEShaderMap>& OutShaderMap, 
+		bool bApplyCompletedShaderMapForRendering,
+		bool bSynchronous = false);
+#endif // WITH_EDITOR
+
+	FString FriendlyName;
+
+	/** Asset using this resource */
+	FName AssetPath;
+
+	/** 
+	 * Contains the compiling id of this shader map when it is being compiled asynchronously. 
+	 * This can be used to access the shader map during async compiling, since GameThreadShaderMap will not have been set yet.
+	 */
+	TArray<int32, TInlineAllocator<1> > OutstandingCompileShaderMapIds;
+
+	/** Feature level that this kernel is representing. */
+	ERHIFeatureLevel::Type FeatureLevel = ERHIFeatureLevel::SM5;
+
+	uint32 bLoadedCookedShaderMapId : 1;
+	FNNERuntimeIREEShaderMapId CookedShaderMapId;
+};
\ No newline at end of file
diff --git a/Engine/Plugins/Experimental/NNERuntimeIREE/Source/NNERuntimeIREEShader/Internal/NNERuntimeIREEShaderType.h b/Engine/Plugins/Experimental/NNERuntimeIREE/Source/NNERuntimeIREEShader/Internal/NNERuntimeIREEShaderType.h
new file mode 100644
--- /dev/null
+++ b/Engine/Plugins/Experimental/NNERuntimeIREE/Source/NNERuntimeIREEShader/Internal/NNERuntimeIREEShaderType.h
@@ -0,0 +1,156 @@
+// Copyright Epic Games, Inc. All Rights Reserved.
+
+#pragma once
+
+#include "CoreMinimal.h"
+#include "Engine/EngineTypes.h"
+#include "GlobalShader.h"
+#include "Shader.h"
+#if WITH_EDITOR
+#include "ShaderCompiler.h"
+#endif
+
+class FShaderCommonCompileJob;
+class FShaderCompileJob;
+
+class FNNERuntimeIREEResource;
+
+
+struct FNNERuntimeIREEShaderPermutationParameters : public FShaderPermutationParameters
+{
+	FNNERuntimeIREEShaderPermutationParameters(EShaderPlatform InPlatform)
+		: FShaderPermutationParameters(InPlatform)
+	{
+	}
+};
+
+class FNNERuntimeIREEShaderType : public FShaderType
+{
+public:
+	struct FParameters : public FShaderType::FParameters
+	{
+		const FShaderParametersMetadata& ShaderParamMetadata;
+
+		FParameters(
+			const FShaderParametersMetadata& InShaderParamMetadata
+			)
+			: ShaderParamMetadata(InShaderParamMetadata)
+		{
+		}
+	};
+
+	struct CompiledShaderInitializerType : FGlobalShaderType::CompiledShaderInitializerType
+	{
+		const FString DebugDescription;
+
+		CompiledShaderInitializerType(
+			const FShaderType* InType,
+			const FParameters* InParameters,
+			int32 InPermutationId,
+			const FShaderCompilerOutput& CompilerOutput,
+			const FSHAHash& InKernelShaderMapHash,
+			const FString& InDebugDescription
+			)
+			: FGlobalShaderType::CompiledShaderInitializerType(
+				InType,
+				InParameters,
+				InPermutationId,
+				CompilerOutput, 
+				InKernelShaderMapHash,
+				nullptr,
+				nullptr
+				)
+			, DebugDescription(InDebugDescription)
+		{
+		}
+	};
+
+	FNNERuntimeIREEShaderType(
+		FTypeLayoutDesc& InTypeLayout,
+		const TCHAR* InName,
+		const TCHAR* InSourceFilename,
+		const TCHAR* InFunctionName,
+		uint32 InFrequency,					// ugly - ignored for but needed for IMPLEMENT_SHADER_TYPE macro magic
+		int32 InTotalPermutationCount,
+		ConstructSerializedType InConstructSerializedRef,
+		ConstructCompiledType InConstructCompiledRef,
+		ShouldCompilePermutationType InShouldCompilePermutationRef,
+		ShouldPrecachePermutationType InShouldPrecachePermutationRef,
+		GetRayTracingPayloadTypeType InGetRayTracingPayloadTypeRef,
+		GetShaderBindingLayoutType InGetShaderBindingLayoutTypeRef,
+#if WITH_EDITOR
+		ModifyCompilationEnvironmentType InModifyCompilationEnvironmentRef,
+		ValidateCompiledResultType InValidateCompiledResultRef,
+		GetOverrideJobPriorityType InGetOverrideJobPriorityRef,
+#endif // WITH_EDITOR
+		uint32 InTypeSize,
+		const FShaderParametersMetadata* InRootParametersMetadata = nullptr
+		):
+		FShaderType(
+			EShaderTypeForDynamicCast::NNERuntimeIREE, 
+			InTypeLayout, 
+			InName, 
+			InSourceFilename, 
+			InFunctionName, 
+			SF_Compute, 
+			InTotalPermutationCount,
+			InConstructSerializedRef,
+			InConstructCompiledRef,
+			InShouldCompilePermutationRef,
+			InShouldPrecachePermutationRef,
+			InGetRayTracingPayloadTypeRef,
+			InGetShaderBindingLayoutTypeRef,
+#if WITH_EDITOR
+			InModifyCompilationEnvironmentRef,
+			InValidateCompiledResultRef,
+			InGetOverrideJobPriorityRef,
+#endif // WITH_EDITOR
+			InTypeSize,
+			InRootParametersMetadata
+			)
+	{
+	}
+
+#if WITH_EDITOR
+	void BeginCompileShader(
+		uint32 ShaderMapId,
+		int32 PermutationId,
+		const FNNERuntimeIREEResource* InKernel,
+		FSharedShaderCompilerEnvironment* InCompilationEnvironment,
+		EShaderPlatform Platform,
+		TArray<FShaderCommonCompileJobPtr>& InOutNewJobs,
+		FShaderTarget Target
+		);
+
+	/**
+	 * Either creates a new instance of this type or returns an equivalent existing shader.
+	 */
+	FShader* FinishCompileShader(
+		const FSHAHash& InKernelShaderMapHash,
+		const FShaderCompileJob& CurrentJob,
+		const FString& InDebugDescription
+		) const;
+#endif // WITH_EDITOR
+
+	/**
+	 * Checks if the shader type should be cached.
+	 * @return True if this shader type should be cached.
+	 */
+	bool ShouldCache(EShaderPlatform InPlatform, const FNNERuntimeIREEResource* Kernel) const
+	{
+		return ShouldCompilePermutation(FNNERuntimeIREEShaderPermutationParameters(InPlatform));
+	}
+
+#if WITH_EDITOR
+protected:
+	/**
+	 * Sets up the environment used to compile an instance of this shader type.
+	 * @param InPlatform - Platform to compile for.
+	 * @param OutEnvironment - The shader compile environment that the function modifies.
+	 */
+	void SetupCompileEnvironment(EShaderPlatform InPlatform, const FNNERuntimeIREEResource* Kernel, FShaderCompilerEnvironment& OutEnvironment) const
+	{
+		ModifyCompilationEnvironment(FNNERuntimeIREEShaderPermutationParameters(InPlatform), OutEnvironment);
+	}
+#endif // WITH_EDITOR
+};
\ No newline at end of file
diff --git a/Engine/Plugins/Experimental/NNERuntimeIREE/Source/NNERuntimeIREEShader/NNERuntimeIREEShader.Build.cs b/Engine/Plugins/Experimental/NNERuntimeIREE/Source/NNERuntimeIREEShader/NNERuntimeIREEShader.Build.cs
new file mode 100644
--- /dev/null
+++ b/Engine/Plugins/Experimental/NNERuntimeIREE/Source/NNERuntimeIREEShader/NNERuntimeIREEShader.Build.cs
@@ -0,0 +1,24 @@
+// Copyright Epic Games, Inc. All Rights Reserved.
+
+using UnrealBuildTool;
+
+public class NNERuntimeIREEShader : ModuleRules
+{
+	public NNERuntimeIREEShader(ReadOnlyTargetRules Target) : base(Target)
+	{
+		PCHUsage = PCHUsageMode.UseExplicitOrSharedPCHs;
+
+		PrivateDependencyModuleNames.AddRange
+		(
+			new string[]
+			{
+				"Core",
+				"CoreUObject",
+				"Engine",
+				"Projects",
+				"RenderCore",
+				"RHI",
+			}
+		);
+	}
+}
diff --git a/Engine/Plugins/Experimental/NNERuntimeIREE/Source/NNERuntimeIREEShader/Private/NNERuntimeIREEShader.cpp b/Engine/Plugins/Experimental/NNERuntimeIREE/Source/NNERuntimeIREEShader/Private/NNERuntimeIREEShader.cpp
new file mode 100644
--- /dev/null
+++ b/Engine/Plugins/Experimental/NNERuntimeIREE/Source/NNERuntimeIREEShader/Private/NNERuntimeIREEShader.cpp
@@ -0,0 +1,15 @@
+// Copyright Epic Games, Inc. All Rights Reserved.
+
+#include "NNERuntimeIREEShader.h"
+
+#include "NNERuntimeIREEShaderType.h"
+
+IMPLEMENT_SHADER_TYPE(, FNNERuntimeIREEShader, TEXT("/Plugin/NNERuntimeIREEShader/NNERuntimeIREEShader.usf"), TEXT("__"), SF_Compute)
+
+FNNERuntimeIREEShader::FNNERuntimeIREEShader(const FNNERuntimeIREEShaderType::CompiledShaderInitializerType& Initializer)
+	: FShader(Initializer)
+{
+	const FShaderParametersMetadata& ShaderParametersMetadata = static_cast<const FNNERuntimeIREEShaderType::FParameters*>(Initializer.Parameters)->ShaderParamMetadata;
+
+	Bindings.BindForLegacyShaderParameters(this, Initializer.PermutationId, Initializer.ParameterMap, ShaderParametersMetadata, true);
+}
\ No newline at end of file
diff --git a/Engine/Plugins/Experimental/NNERuntimeIREE/Source/NNERuntimeIREEShader/Private/NNERuntimeIREEShaderCompilationManager.cpp b/Engine/Plugins/Experimental/NNERuntimeIREE/Source/NNERuntimeIREEShader/Private/NNERuntimeIREEShaderCompilationManager.cpp
new file mode 100644
--- /dev/null
+++ b/Engine/Plugins/Experimental/NNERuntimeIREE/Source/NNERuntimeIREEShader/Private/NNERuntimeIREEShaderCompilationManager.cpp
@@ -0,0 +1,414 @@
+// Copyright Epic Games, Inc. All Rights Reserved.
+
+#include "NNERuntimeIREEShaderCompilationManager.h"
+
+#include "DataDrivenShaderPlatformInfo.h"
+#include "HAL/FileManager.h"
+#include "Misc/Paths.h"
+#include "NNERuntimeIREEShaderLog.h"
+#include "NNERuntimeIREEShaderShared.h"
+#include "RHIShaderFormatDefinitions.inl"
+#include "ShaderCompiler.h"
+
+#if WITH_EDITOR
+#include "Interfaces/IShaderFormat.h"
+#include "Interfaces/ITargetPlatformManagerModule.h"
+#endif
+
+static int32 GShowComputeKernelShaderWarnings = 0;
+static FAutoConsoleVariableRef CVarShowComputeKernelShaderWarnings(
+	TEXT("ComputeKernel.ShowShaderCompilerWarnings"),
+	GShowComputeKernelShaderWarnings,
+	TEXT("When set to 1, will display all warnings from ComputeKernel shader compiles.")
+	);
+
+
+#if WITH_EDITOR
+FNNERuntimeIREEShaderCompilationManager GComputeKernelShaderCompilationManager;
+
+void FNNERuntimeIREEShaderCompilationManager::Tick(float DeltaSeconds)
+{
+	ProcessAsyncResults();
+}
+
+void FNNERuntimeIREEShaderCompilationManager::AddJobs(TArray<FShaderCommonCompileJobPtr> InNewJobs)
+{
+	check(IsInGameThread());
+	JobQueue.Append(InNewJobs);
+
+	for (FShaderCommonCompileJobPtr& Job : InNewJobs)
+	{
+		FNNERuntimeIREEShaderMapCompileResults& ShaderMapInfo = ComputeKernelShaderMapJobs.FindOrAdd(Job->Id);
+		ShaderMapInfo.NumJobsQueued++;
+
+		FShaderCompileJob* CurrentJob = Job->GetSingleShaderJob();
+
+		CurrentJob->Input.DumpDebugInfoRootPath = GShaderCompilingManager->GetAbsoluteShaderDebugInfoDirectory() / CurrentJob->Input.ShaderPlatformName.ToString();
+		FPaths::NormalizeDirectoryName(CurrentJob->Input.DumpDebugInfoRootPath);
+
+		CurrentJob->Input.DebugExtension.Empty();
+		CurrentJob->Input.DumpDebugInfoPath.Empty();
+		if (GShaderCompilingManager->GetDumpShaderDebugInfo() == FShaderCompilingManager::EDumpShaderDebugInfo::Always)
+		{
+			CurrentJob->Input.DumpDebugInfoRootPath = GShaderCompilingManager->CreateShaderDebugInfoPath(CurrentJob->Input);
+		}
+	}
+
+	GShaderCompilingManager->SubmitJobs(InNewJobs, FString(), FString());
+}
+
+void FNNERuntimeIREEShaderCompilationManager::ProcessAsyncResults()
+{
+	check(IsInGameThread());
+
+	// Process the results from the shader compile worker
+	for (int32 JobIndex = JobQueue.Num() - 1; JobIndex >= 0; JobIndex--)
+	{
+		auto CurrentJob = JobQueue[JobIndex]->GetSingleShaderJob();
+
+		if (!CurrentJob->bReleased)
+		{
+			continue;
+		}
+
+		CurrentJob->bSucceeded = CurrentJob->Output.bSucceeded;
+		if (CurrentJob->Output.bSucceeded)
+		{
+			UE_LOG(LogNNERuntimeIREEShader, Verbose, TEXT("GPU shader compile succeeded. Id %d"), CurrentJob->Id);
+		}
+		else
+		{
+			UE_LOG(LogNNERuntimeIREEShader, Verbose, TEXT("GPU shader compile failed! Id %d"), CurrentJob->Id);
+		}
+
+		FNNERuntimeIREEShaderMapCompileResults& ShaderMapResults = ComputeKernelShaderMapJobs.FindChecked(CurrentJob->Id);
+		ShaderMapResults.FinishedJobs.Add(JobQueue[JobIndex]);
+		ShaderMapResults.bAllJobsSucceeded = ShaderMapResults.bAllJobsSucceeded && CurrentJob->bSucceeded;
+		JobQueue.RemoveAt(JobIndex);
+	}
+
+	for (TMap<int32, FNNERuntimeIREEShaderMapCompileResults>::TIterator It(ComputeKernelShaderMapJobs); It; ++It)
+	{
+		const FNNERuntimeIREEShaderMapCompileResults& Results = It.Value();
+
+		if (Results.FinishedJobs.Num() == Results.NumJobsQueued)
+		{
+			PendingFinalizeComputeKernelShaderMaps.Add(It.Key(), FNNERuntimeIREEShaderMapCompileResults(Results));
+			ComputeKernelShaderMapJobs.Remove(It.Key());
+		}
+	}
+
+	if (PendingFinalizeComputeKernelShaderMaps.Num() > 0)
+	{
+		ProcessCompiledComputeKernelShaderMaps(PendingFinalizeComputeKernelShaderMaps, 10);
+	}
+}
+
+static bool ParseShaderCompilerError(FShaderCompilerError const& InError, bool bInCompilationSucceeded, FNNERuntimeIREEShaderCompileMessage& OutMessage)
+{
+	FShaderCompilerError Error = InError;
+
+	if (Error.StrippedErrorMessage.RemoveFromStart(TEXT("error: ")))
+	{
+		OutMessage.Type = FNNERuntimeIREEShaderCompileMessage::EMessageType::Error;
+	}
+	else if (Error.StrippedErrorMessage.RemoveFromStart(TEXT("warning: ")))
+	{
+		OutMessage.Type = FNNERuntimeIREEShaderCompileMessage::EMessageType::Warning;
+	}
+	else if (Error.StrippedErrorMessage.RemoveFromStart(TEXT("note: ")))
+	{
+		OutMessage.Type = FNNERuntimeIREEShaderCompileMessage::EMessageType::Info;
+	}
+	else
+	{
+		// General rule for preprocessing errors - if compilation succeeded errors are warnings, otherwise errors.
+		OutMessage.Type = bInCompilationSucceeded ? FNNERuntimeIREEShaderCompileMessage::EMessageType::Warning : FNNERuntimeIREEShaderCompileMessage::EMessageType::Error;
+	}
+
+	OutMessage.Text = Error.StrippedErrorMessage;
+	OutMessage.VirtualFilePath = Error.ErrorVirtualFilePath;
+
+	// Fix up the DataInterface generated file paths before any error reporting.
+	// Magic path structure is set in ComputeGraph compilation.
+	if (OutMessage.VirtualFilePath.RemoveFromStart(TEXT("/Engine/Generated/DataInterface/")))
+	{
+		OutMessage.VirtualFilePath.MidInline(OutMessage.VirtualFilePath.Find(TEXT("/")));
+	}
+	// Store any disk paths before error reporting. Can skip some known cases that won't have disk paths.
+	if (OutMessage.VirtualFilePath.StartsWith(TEXT("/")) && !OutMessage.VirtualFilePath.StartsWith(TEXT("/Engine/Generated/")))
+	{
+		OutMessage.RealFilePath = GetShaderSourceFilePath(OutMessage.VirtualFilePath);
+	}
+
+	// Populate line and column number strings if available.
+	FString Line, Column;
+	// Check for "line,col" format first.
+	if (!Error.ErrorLineString.Split(TEXT(","), &Line, &Column))
+	{
+		// Fall back to "line" which is logged by preprocessor errors.
+		if (Error.ErrorLineString.IsNumeric())
+		{
+			Line = Error.ErrorLineString;
+			Column.Empty();
+		}
+		else
+		{
+			Line.Empty();
+			Column.Empty();
+		}
+	}
+
+	if (Line.IsNumeric())
+	{
+		LexFromString(OutMessage.Line, *Line);
+	}
+
+	if (Column.IsNumeric())
+	{
+		LexFromString(OutMessage.ColumnStart, *Column);
+		OutMessage.ColumnEnd = OutMessage.ColumnStart;
+
+		if (Error.HasLineMarker())
+		{
+			for (TCHAR Character : Error.HighlightedLineMarker)
+			{
+				OutMessage.ColumnEnd += (Character == TEXT('~')) ? 1 : 0;
+			}
+		}
+	}
+
+	return true;
+}
+
+static void LogShaderCompilerErrors(FNNERuntimeIREEShaderCompileResults const& Results)
+{
+	for (FNNERuntimeIREEShaderCompileMessage const& Message : Results.Messages)
+	{
+		FString Path = Message.RealFilePath.IsEmpty() ? Message.VirtualFilePath : Message.RealFilePath;
+
+		const bool bPreparePathForVisualStudioHotlink = !Message.RealFilePath.IsEmpty() && FPlatformMisc::IsDebuggerPresent();
+		if (bPreparePathForVisualStudioHotlink)
+		{
+			// Convert path to absolute, and prepend with newline so that it is clickable in Visual Studio.
+			Path = TEXT("\n") + IFileManager::Get().ConvertToAbsolutePathForExternalAppForRead(*Message.RealFilePath);
+		}
+
+		FString Line;
+		if (Message.ColumnStart == Message.ColumnEnd)
+		{
+			Line = FString::Printf(TEXT("(%d,%d)"), Message.Line, Message.ColumnStart);
+		}
+		else
+		{
+			Line = FString::Printf(TEXT("(%d,%d-%d)"), Message.Line, Message.ColumnStart, Message.ColumnEnd);
+		}
+
+		FString MessageText = FString::Printf(TEXT("%s%s: %s"), *Path, *Line, *Message.Text);
+
+		if (Message.Type == FNNERuntimeIREEShaderCompileMessage::EMessageType::Warning)
+		{
+			UE_LOG(LogNNERuntimeIREEShader, Warning, TEXT("%s"), *MessageText);
+		}
+		else if (Message.Type == FNNERuntimeIREEShaderCompileMessage::EMessageType::Error)
+		{
+			UE_LOG(LogNNERuntimeIREEShader, Error, TEXT("%s"), *MessageText);
+		}
+	}
+}
+
+void FNNERuntimeIREEShaderCompilationManager::ProcessCompiledComputeKernelShaderMaps(TMap<int32, FNNERuntimeIREEShaderMapFinalizeResults>& CompiledShaderMaps, float TimeBudget)
+{
+	// Keeps shader maps alive as they are passed from the shader compiler and applied to the owning kernel
+	TArray<TRefCountPtr<FNNERuntimeIREEShaderMap> > LocalShaderMapReferences;
+	TMap<FNNERuntimeIREEResource*, FNNERuntimeIREEShaderMap*> KernelsToUpdate;
+
+	// Process compiled shader maps in FIFO order, in case a shader map has been enqueued multiple times,
+	// Which can happen if a kernel is edited while a background compile is going on
+	for (TMap<int32, FNNERuntimeIREEShaderMapFinalizeResults>::TIterator ProcessIt(CompiledShaderMaps); ProcessIt; ++ProcessIt)
+	{
+		TRefCountPtr<FNNERuntimeIREEShaderMap> ShaderMap = nullptr;
+		TArray<FNNERuntimeIREEResource*>* Kernels = nullptr;
+
+		for (TMap<TRefCountPtr<FNNERuntimeIREEShaderMap>, TArray<FNNERuntimeIREEResource*> >::TIterator ShaderMapIt(FNNERuntimeIREEShaderMap::GetInFlightShaderMaps()); ShaderMapIt; ++ShaderMapIt)
+		{
+			if (ShaderMapIt.Key()->GetCompilingId() == ProcessIt.Key())
+			{
+				ShaderMap = ShaderMapIt.Key();
+				Kernels = &ShaderMapIt.Value();
+				break;
+			}
+		}
+
+		if (ShaderMap && Kernels)
+		{
+			FNNERuntimeIREEShaderMapFinalizeResults& CompileResults = ProcessIt.Value();
+			TArray<FShaderCommonCompileJobPtr>& ResultArray = CompileResults.FinishedJobs;
+			FNNERuntimeIREEShaderCompileResults ProcessedCompileResults;
+
+			// Make a copy of the array as this entry of FNNERuntimeIREEShaderMap::ShaderMapsBeingCompiled will be removed below
+			TArray<FNNERuntimeIREEResource*> KernelsArray = *Kernels;
+			bool bSuccess = true;
+
+			for (int32 JobIndex = 0; JobIndex < ResultArray.Num(); JobIndex++)
+			{
+				FShaderCompileJob& CurrentJob = static_cast<FShaderCompileJob&>(*ResultArray[JobIndex].GetReference());
+				bSuccess = bSuccess && CurrentJob.bSucceeded;
+
+				if (bSuccess)
+				{
+					check(CurrentJob.Output.ShaderCode.GetShaderCodeSize() > 0);
+				}
+
+				if (GShowComputeKernelShaderWarnings || !CurrentJob.bSucceeded)
+				{
+					TArray<FShaderCompilerError> Errors = CurrentJob.Output.Errors;
+					FShaderCompilerError::ExtractSourceLocations(Errors);
+
+					for (int32 ErrorIndex = 0; ErrorIndex < Errors.Num(); ErrorIndex++)
+					{
+						FNNERuntimeIREEShaderCompileMessage Message;
+						if (ParseShaderCompilerError(Errors[ErrorIndex], CurrentJob.Output.bSucceeded, Message))
+						{
+							ProcessedCompileResults.Messages.AddUnique(Message);
+						}
+					}
+
+					if (ProcessedCompileResults.Messages.Num())
+					{
+						UE_LOG(LogNNERuntimeIREEShader, Verbose, TEXT("There were errors for job \"%s\""), *CurrentJob.Input.DebugGroupName)
+					}
+				}
+				else
+				{
+					UE_LOG(LogNNERuntimeIREEShader, Verbose, TEXT("There were NO errors for job \"%s\""), *CurrentJob.Input.DebugGroupName);
+				}
+			}
+
+			bool bShaderMapComplete = true;
+
+			if (bSuccess)
+			{
+				bShaderMapComplete = ShaderMap->ProcessCompilationResults(ResultArray, CompileResults.FinalizeJobIndex, TimeBudget);
+			}
+
+			if (bShaderMapComplete)
+			{
+				ShaderMap->SetCompiledSuccessfully(bSuccess);
+
+				// Pass off the reference of the shader map to LocalShaderMapReferences
+				LocalShaderMapReferences.Add(ShaderMap);
+				FNNERuntimeIREEShaderMap::GetInFlightShaderMaps().Remove(ShaderMap);
+
+				for (FNNERuntimeIREEResource* Kernel : KernelsArray)
+				{
+					FNNERuntimeIREEShaderMap* CompletedShaderMap = ShaderMap;
+
+					Kernel->RemoveOutstandingCompileId(ShaderMap->GetCompilingId());
+
+					// Only process results that still match the ID which requested a compile
+					// This avoids applying shadermaps which are out of date and a newer one is in the async compiling pipeline
+					if (Kernel->IsSame(CompletedShaderMap->GetShaderMapId()))
+					{
+						if (!bSuccess)
+						{
+							// Propagate error messages
+							LogShaderCompilerErrors(ProcessedCompileResults);
+							Kernel->SetCompilationResults(ProcessedCompileResults);
+							KernelsToUpdate.Add(Kernel, nullptr);
+						}
+						else
+						{
+							// if we succeeded and our shader map is not complete this could be because the kernel was being edited quicker then the compile could be completed
+							// Don't modify kernel for which the compiled shader map is no longer complete
+							// This shouldn't happen since kernels are pretty much baked in the designated config file.
+							if (CompletedShaderMap->IsComplete(Kernel, true))
+							{
+								KernelsToUpdate.Add(Kernel, CompletedShaderMap);
+							}
+
+							if (GShowComputeKernelShaderWarnings && ProcessedCompileResults.Messages.Num() > 0)
+							{
+								UE_LOG(LogNNERuntimeIREEShader, Warning, TEXT("Warnings while compiling ComputeKernel %s for platform %s:"),
+									*Kernel->GetFriendlyName(),
+									*LegacyShaderPlatformToShaderFormat(ShaderMap->GetShaderPlatform()).ToString());
+
+								LogShaderCompilerErrors(ProcessedCompileResults);
+								Kernel->SetCompilationResults(ProcessedCompileResults);
+							}
+						}
+					}
+					else
+					{
+						if (CompletedShaderMap->IsComplete(Kernel, true))
+						{
+							const FString ShaderFormatName = FDataDrivenShaderPlatformInfo::GetShaderFormat(ShaderMap->GetShaderPlatform()).ToString();
+							if (bSuccess)
+							{
+								const FString SuccessMessage = FString::Printf(TEXT("%s: %s shader compilation success!"), *Kernel->GetFriendlyName(), *ShaderFormatName);
+								Kernel->NotifyCompilationFinished(SuccessMessage);
+							}
+							else
+							{
+								const FString FailMessage = FString::Printf(TEXT("%s: %s shader compilation failed."), *Kernel->GetFriendlyName(), *ShaderFormatName);
+								Kernel->NotifyCompilationFinished(FailMessage);
+							}
+						}
+					}
+				}
+
+				// Cleanup shader jobs and compile tracking structures
+				ResultArray.Empty();
+				CompiledShaderMaps.Remove(ShaderMap->GetCompilingId());
+			}
+
+			if (TimeBudget < 0)
+			{
+				break;
+			}
+		}
+	}
+
+	if (KernelsToUpdate.Num() > 0)
+	{
+		for (TMap<FNNERuntimeIREEResource*, FNNERuntimeIREEShaderMap*>::TConstIterator It(KernelsToUpdate); It; ++It)
+		{
+			FNNERuntimeIREEResource* Kernel = It.Key();
+			FNNERuntimeIREEShaderMap* ShaderMap = It.Value();
+
+			Kernel->SetGameThreadShaderMap(It.Value());
+
+			// ENQUEUE_RENDER_COMMAND(FSetShaderMapOnComputeKernel)(
+			// 	[Kernel, ShaderMap](FRHICommandListImmediate& RHICmdList)
+			// 	{
+			// 		Kernel->SetRenderingThreadShaderMap(ShaderMap);
+			// 	});
+
+			if (ShaderMap && ShaderMap->CompiledSuccessfully())
+			{
+				const FString ShaderFormatName = FDataDrivenShaderPlatformInfo::GetShaderFormat(ShaderMap->GetShaderPlatform()).ToString();
+				const FString SuccessMessage = FString::Printf(TEXT("%s: %s shader compilation success!"), 
+					*Kernel->GetFriendlyName(), 
+					*ShaderFormatName);
+				Kernel->NotifyCompilationFinished(SuccessMessage);
+			}
+			else
+			{
+				const FString FailMessage = FString::Printf(TEXT("%s: Shader compilation failed."), *Kernel->GetFriendlyName());
+				Kernel->NotifyCompilationFinished(FailMessage);
+			}
+		}
+	}
+}
+
+void FNNERuntimeIREEShaderCompilationManager::FinishCompilation(const TCHAR* InKernelName, const TArray<int32>& ShaderMapIdsToFinishCompiling)
+{
+	check(!FPlatformProperties::RequiresCookedData());
+
+	GShaderCompilingManager->FinishCompilation(NULL, ShaderMapIdsToFinishCompiling);
+	ProcessAsyncResults();	// grab compiled shader maps and assign them to their resources
+
+	check(ComputeKernelShaderMapJobs.Num() == 0);
+}
+
+#endif // WITH_EDITOR
diff --git a/Engine/Plugins/Experimental/NNERuntimeIREE/Source/NNERuntimeIREEShader/Private/NNERuntimeIREEShaderCompilationManager.h b/Engine/Plugins/Experimental/NNERuntimeIREE/Source/NNERuntimeIREEShader/Private/NNERuntimeIREEShaderCompilationManager.h
new file mode 100644
--- /dev/null
+++ b/Engine/Plugins/Experimental/NNERuntimeIREE/Source/NNERuntimeIREEShader/Private/NNERuntimeIREEShaderCompilationManager.h
@@ -0,0 +1,60 @@
+// Copyright Epic Games, Inc. All Rights Reserved.
+
+#pragma once
+
+#include "NNERuntimeIREEShaderShared.h"
+#include "ShaderCompiler.h"
+
+/** Results for a single compiled shader map. */
+struct FNNERuntimeIREEShaderMapCompileResults
+{
+	FNNERuntimeIREEShaderMapCompileResults() :
+		NumJobsQueued(0),
+		bAllJobsSucceeded(true),
+		bRecreateComponentRenderStateOnCompletion(false)
+	{}
+
+	int32 NumJobsQueued;
+	bool bAllJobsSucceeded;
+	bool bRecreateComponentRenderStateOnCompletion;
+	TArray<FShaderCommonCompileJobPtr> FinishedJobs;
+};
+
+/** Results for a single compiled and finalized shader map. */
+struct FNNERuntimeIREEShaderMapFinalizeResults : public FNNERuntimeIREEShaderMapCompileResults
+{
+	/** Tracks finalization progress on this shader map. */
+	int32 FinalizeJobIndex;
+
+	FNNERuntimeIREEShaderMapFinalizeResults(const FNNERuntimeIREEShaderMapCompileResults& InCompileResults) :
+		FNNERuntimeIREEShaderMapCompileResults(InCompileResults),
+		FinalizeJobIndex(0)
+	{}
+};
+
+#if WITH_EDITOR
+
+/** Handles finished shader compile jobs, applying of the shaders to their config asset, and some error handling. */
+class FNNERuntimeIREEShaderCompilationManager
+{
+public:
+	void Tick(float DeltaSeconds);
+	void AddJobs(TArray<FShaderCommonCompileJobPtr> InNewJobs);
+	void FinishCompilation(const TCHAR* InKernelName, const TArray<int32>& ShaderMapIdsToFinishCompiling);
+
+private:
+	void ProcessAsyncResults();
+	void ProcessCompiledComputeKernelShaderMaps(TMap<int32, FNNERuntimeIREEShaderMapFinalizeResults>& CompiledShaderMaps, float TimeBudget);
+
+	TArray<FShaderCommonCompileJobPtr> JobQueue;
+
+	/** Map from shader map Id to the compile results for that map, used to gather compiled results. */
+	TMap<int32, FNNERuntimeIREEShaderMapCompileResults> ComputeKernelShaderMapJobs;
+
+	/** Map from shader map id to results being finalized.  Used to track shader finalizations over multiple frames. */
+	TMap<int32, FNNERuntimeIREEShaderMapFinalizeResults> PendingFinalizeComputeKernelShaderMaps;
+};
+
+extern FNNERuntimeIREEShaderCompilationManager GComputeKernelShaderCompilationManager;
+
+#endif // WITH_EDITOR
diff --git a/Engine/Plugins/Experimental/NNERuntimeIREE/Source/NNERuntimeIREEShader/Private/NNERuntimeIREEShaderMap.cpp b/Engine/Plugins/Experimental/NNERuntimeIREE/Source/NNERuntimeIREEShader/Private/NNERuntimeIREEShaderMap.cpp
new file mode 100644
--- /dev/null
+++ b/Engine/Plugins/Experimental/NNERuntimeIREE/Source/NNERuntimeIREEShader/Private/NNERuntimeIREEShaderMap.cpp
@@ -0,0 +1,518 @@
+// Copyright Epic Games, Inc. All Rights Reserved.
+
+#include "NNERuntimeIREEShader.h"
+
+#include "NNERuntimeIREEShaderShared.h"
+#include "NNERuntimeIREEShaderCompilationManager.h"
+#include "NNERuntimeIREEShaderLog.h"
+#include "Interfaces/ITargetPlatformManagerModule.h"
+#include "ProfilingDebugging/CookStats.h"
+#include "ProfilingDebugging/DiagnosticTable.h"
+#include "RHIShaderFormatDefinitions.inl"
+#include "Serialization/MemoryWriter.h"
+#include "Serialization/MemoryReader.h"
+#include "ShaderCompiler.h"
+#include "Stats/StatsMisc.h"
+#include "UObject/UObjectGlobals.h"
+#include "DataDrivenShaderPlatformInfo.h"
+
+//
+// Globals
+//
+TMap<FNNERuntimeIREEShaderMapId, FNNERuntimeIREEShaderMap*> FNNERuntimeIREEShaderMap::GIdToKernelShaderMap[SP_NumPlatforms];
+TArray<FNNERuntimeIREEShaderMap*> FNNERuntimeIREEShaderMap::AllKernelShaderMaps;
+
+#if WITH_EDITOR
+/** 
+ * Tracks FNNERuntimeIREEResource and their shader maps that are being compiled.
+ * Uses a TRefCountPtr as this will be the only reference to a shader map while it is being compiled.
+ */
+TMap<TRefCountPtr<FNNERuntimeIREEShaderMap>, TArray<FNNERuntimeIREEResource*> > FNNERuntimeIREEShaderMap::ShaderMapsBeingCompiled;
+#endif // WITH_EDITOR
+
+static inline bool ShouldCacheNNERuntimeIREEShader(const FNNERuntimeIREEShaderType* InShaderType, EShaderPlatform InPlatform, const FNNERuntimeIREEResource* InKernel)
+{
+	return InShaderType->ShouldCache(InPlatform, InKernel) && InKernel->ShouldCache(InPlatform, InShaderType);
+}
+
+/** Hashes the kernel specific part of this shader map Id. */
+void FNNERuntimeIREEShaderMapId::GetKernelHash(FSHAHash& OutHash) const
+{
+	FSHA1 HashState;
+
+	HashState.Update((const uint8*)&ShaderCodeHash, sizeof(ShaderCodeHash));
+	HashState.Update((const uint8*)&FeatureLevel, sizeof(FeatureLevel));
+
+	HashState.Final();
+	HashState.GetHash(&OutHash.Hash[0]);
+}
+
+/** 
+* Tests this set against another for equality, disregarding override settings.
+* 
+* @param ReferenceSet	The set to compare against
+* @return				true if the sets are equal
+*/
+bool FNNERuntimeIREEShaderMapId::operator==(const FNNERuntimeIREEShaderMapId& InReferenceSet) const
+{
+	if (  ShaderCodeHash != InReferenceSet.ShaderCodeHash
+		|| FeatureLevel != InReferenceSet.FeatureLevel)
+	{
+		return false;
+	}
+
+	if (ShaderTypeDependencies.Num() != InReferenceSet.ShaderTypeDependencies.Num())
+	{
+		return false;
+	}
+
+	if (LayoutParams != InReferenceSet.LayoutParams)
+	{
+		return false;
+	}
+
+	for (int32 ShaderIndex = 0; ShaderIndex < ShaderTypeDependencies.Num(); ShaderIndex++)
+	{
+		const FShaderTypeDependency& ShaderTypeDependency = ShaderTypeDependencies[ShaderIndex];
+
+		if (ShaderTypeDependency != InReferenceSet.ShaderTypeDependencies[ShaderIndex])
+		{
+			return false;
+		}
+	}
+
+	return true;
+}
+
+#if WITH_EDITOR
+
+/**
+ * Enqueues a compilation for a new shader of this type.
+ * @param InKernel - The kernel to link the shader with.
+ */
+void FNNERuntimeIREEShaderType::BeginCompileShader(
+	uint32 InShaderMapId,
+	int32 PermutationId,
+	const FNNERuntimeIREEResource* InKernel,
+	FSharedShaderCompilerEnvironment* InCompilationEnvironment,
+	EShaderPlatform InPlatform,
+	TArray<FShaderCommonCompileJobPtr>& NewJobs,
+	FShaderTarget InTarget
+	)
+{
+	FShaderCompileJob* NewJob = GShaderCompilingManager->PrepareShaderCompileJob(
+		InShaderMapId, 
+		FShaderCompileJobKey(this, nullptr, PermutationId),
+		EShaderCompileJobPriority::Normal
+		);
+
+	TCHAR const* SourceFilePath = TEXT("/Plugin/NNERuntimeIREEShader/NNERuntimeIREEShader.usf");
+	TCHAR const* GeneratedFilePath = TEXT("/Plugin/NNERuntimeIREEShader/Generated/NNERuntimeIREEShader.ush");
+
+	NewJob->ShaderParameters = MakeShared<const FParameters, ESPMode::ThreadSafe>(*InKernel->GetShaderParamMetadata());
+	NewJob->Input.SharedEnvironment = InCompilationEnvironment;
+	NewJob->Input.Target = InTarget;
+	NewJob->Input.ShaderFormat = LegacyShaderPlatformToShaderFormat(InPlatform);
+	NewJob->Input.VirtualSourceFilePath = SourceFilePath;
+	NewJob->Input.EntryPointName = InKernel->GetEntryPoint();
+	NewJob->Input.Environment.IncludeVirtualPathToContentsMap.Add(GeneratedFilePath, InKernel->GetHLSLSource());
+	UE_LOG(LogNNERuntimeIREEShader, Verbose, TEXT("%s"), *InKernel->GetHLSLSource());
+	
+	AddUniformBufferIncludesToEnvironment(NewJob->Input.Environment, InPlatform);
+
+	FShaderCompilerEnvironment& ShaderEnvironment = NewJob->Input.Environment;
+
+	UE_LOG(LogNNERuntimeIREEShader, Verbose, TEXT("			%s"), GetName());
+	// COOK_STAT(ComputeKernelShaderCookStats::ShadersCompiled++);
+
+	// InKernel->SetupCompileEnvironment(PermutationId, ShaderEnvironment);
+
+	// Allow the shader type to modify the compile environment.
+	SetupCompileEnvironment(InPlatform, InKernel, ShaderEnvironment);
+
+	::GlobalBeginCompileShader(
+		InKernel->GetFriendlyName(),
+		nullptr,
+		this,
+		nullptr,//ShaderPipeline,
+		PermutationId,
+		SourceFilePath,
+		*InKernel->GetEntryPoint(),
+		FShaderTarget(GetFrequency(), InPlatform),
+		NewJob->Input
+		);
+
+	NewJobs.Add(FShaderCommonCompileJobPtr(NewJob));
+}
+
+/**
+ * Either creates a new instance of this type or returns an equivalent existing shader.
+ * @param InShaderMapHash - Precomputed hash of the shader map 
+ * @param InCurrentJob - Compile job that was enqueued by BeginCompileShader.
+ */
+FShader* FNNERuntimeIREEShaderType::FinishCompileShader(
+	const FSHAHash& InShaderMapHash,
+	const FShaderCompileJob& InCurrentJob,
+	const FString& InDebugDescription
+	) const
+{
+	check(InCurrentJob.bSucceeded);
+
+	FShader* Shader = ConstructCompiled(
+		FNNERuntimeIREEShaderType::CompiledShaderInitializerType(
+			this,
+			static_cast<const FParameters*>(InCurrentJob.ShaderParameters.Get()),
+			InCurrentJob.Key.PermutationId,
+			InCurrentJob.Output,
+			InShaderMapHash,
+			InDebugDescription
+			)
+		);
+
+	InCurrentJob.Output.ParameterMap.VerifyBindingsAreComplete(GetName(), InCurrentJob.Output.Target, InCurrentJob.Key.VFType);
+
+	return Shader;
+}
+#endif // WITH_EDITOR
+
+/**
+ * Finds the shader map for a kernel.
+ * @param InShaderMapId - The kernel id and static parameter set identifying the shader map
+ * @param InPlatform - The platform to lookup for
+ * @return nullptr if no cached shader map was found.
+ */
+FNNERuntimeIREEShaderMap* FNNERuntimeIREEShaderMap::FindId(const FNNERuntimeIREEShaderMapId& InShaderMapId, EShaderPlatform InPlatform)
+{
+	FNNERuntimeIREEShaderMap* Result = GIdToKernelShaderMap[InPlatform].FindRef(InShaderMapId);
+	check(Result == nullptr || !Result->bDeletedThroughDeferredCleanup);
+	return Result;
+}
+
+#if WITH_EDITOR
+
+/**
+* Compiles the shaders for a kernel and caches them in this shader map.
+* @param InKernel - The kernel to compile shaders for.
+* @param InShaderMapId - the kernel id and set of static parameters to compile for
+* @param InPlatform - The platform to compile to
+*/
+void FNNERuntimeIREEShaderMap::Compile(
+	FNNERuntimeIREEResource* InKernel, 
+	const FNNERuntimeIREEShaderMapId& InShaderMapId,
+	TRefCountPtr<FSharedShaderCompilerEnvironment> InCompilationEnvironment,
+	const FNNERuntimeIREECompilationOutput& InKernelCompilationOutput,
+	EShaderPlatform InPlatform,
+	bool bSynchronousCompile,
+	bool bApplyCompletedShaderMapForRendering
+	)
+{
+	if (FPlatformProperties::RequiresCookedData())
+	{
+		UE_LOG(LogNNERuntimeIREEShader, Fatal, TEXT("Trying to compile ComputeKernel shader %s at run-time, which is not supported on consoles!"), *InKernel->GetFriendlyName() );
+	}
+	else
+	{
+		// Make sure we are operating on a referenced shader map or the below Find will cause this shader map to be deleted,
+		// Since it creates a temporary ref counted pointer.
+		check(NumRefs > 0);
+  
+		// Add this shader map and to ComputeKernelShaderMapsBeingCompiled
+		TArray<FNNERuntimeIREEResource*>* Kernel = ShaderMapsBeingCompiled.Find(this);
+  
+		if (Kernel)
+		{
+			check(!bSynchronousCompile);
+			Kernel->AddUnique(InKernel);
+		}
+		else
+		{
+			// Assign a unique identifier so that shaders from this shader map can be associated with it after a deferred compile
+			CompilingId = FShaderCommonCompileJob::GetNextJobId();
+			UE_LOG(LogNNERuntimeIREEShader, Verbose, TEXT("CompilingId = %p %d"), InKernel, CompilingId);
+			InKernel->AddCompileId(CompilingId);
+
+			// Setup the compilation environment.
+			// InKernel->SetupShaderCompilationEnvironment(InPlatform, *InCompilationEnvironment);
+  
+			// Store the kernel name for debugging purposes.
+			FNNERuntimeIREEShaderMapContent* NewContent = new FNNERuntimeIREEShaderMapContent(InPlatform);
+			NewContent->FriendlyName = InKernel->GetFriendlyName();
+			NewContent->CompilationOutput = InKernelCompilationOutput;
+			NewContent->ShaderMapId = InShaderMapId;
+			AssignContent(NewContent);
+
+			uint32 NumShaders = 0;
+			TArray<FShaderCommonCompileJobPtr> NewJobs;
+	
+			// Iterate over all shader types.
+			for(TLinkedList<FShaderType*>::TIterator ShaderTypeIt(FShaderType::GetTypeList());ShaderTypeIt;ShaderTypeIt.Next())
+			{
+				FNNERuntimeIREEShaderType* ShaderType = ShaderTypeIt->GetNNERuntimeIREEShaderType();
+				if (ShaderType && ShouldCacheNNERuntimeIREEShader(ShaderType, InPlatform, InKernel))
+				{
+					// Verify that the shader map Id contains inputs for any shaders that will be put into this shader map
+					checkf(InShaderMapId.ContainsShaderType(ShaderType), TEXT("IREE kernel shader map %s missing expected shader type %s"), *GetFriendlyName(), ShaderType->GetName());
+					
+					// Compile this NNERuntimeIREE shader.
+					TArray<FString> ShaderErrors;
+  
+					for (int32 PermutationId = 0; PermutationId < InKernel->GetNumPermutations(); ++PermutationId)
+					{
+						// Only compile the shader if we don't already have it
+						if (!NewContent->HasShader(ShaderType, PermutationId))
+						{
+							ShaderType->BeginCompileShader(
+								CompilingId,
+								PermutationId,
+								InKernel,
+								InCompilationEnvironment,
+								InPlatform,
+								NewJobs,
+								FShaderTarget(ShaderType->GetFrequency(), GetShaderPlatform())
+								);
+						}
+						NumShaders++;
+					}
+				}
+				else if (ShaderType)
+				{
+					InKernel->RemoveOutstandingCompileId(CompilingId);
+					const FString ShaderFormatName = FDataDrivenShaderPlatformInfo::GetShaderFormat(InPlatform).ToString();
+
+					FString Message = FString::Printf(TEXT("%s: Compilation not supported on %s."), 
+						*InKernel->GetFriendlyName(), 
+						*ShaderFormatName);
+					// InKernel->NotifyCompilationFinished(Message);
+					UE_LOG(LogNNERuntimeIREEShader, Fatal, TEXT("%s"), *Message);
+
+				}
+			}
+  
+			if (!Kernel)
+			{
+				UE_LOG(LogNNERuntimeIREEShader, Verbose, TEXT("		%u Shaders"), NumShaders);
+			}
+
+			// Register this shader map in the global ComputeKernel->shadermap map
+			Register(InPlatform);
+  
+			// Mark the shader map as not having been finalized with ProcessCompilationResults
+			bCompilationFinalized = false;
+  
+			// Mark as not having been compiled
+			bCompiledSuccessfully = false;
+  
+			if (NumShaders > 0)
+			{
+				GComputeKernelShaderCompilationManager.AddJobs(NewJobs);
+
+				TArray<FNNERuntimeIREEResource*> NewCorrespondingKernels;
+				NewCorrespondingKernels.Add(InKernel);
+				ShaderMapsBeingCompiled.Add(this, NewCorrespondingKernels);
+#if DEBUG_INFINITESHADERCOMPILE
+				UE_LOG(LogTemp, Display, TEXT("Added ComputeKernel ShaderMap 0x%08X%08X with kernel 0x%08X%08X to ComputeKernelShaderMapsBeingCompiled"), (int)((int64)(this) >> 32), (int)((int64)(this)), (int)((int64)(InKernel) >> 32), (int)((int64)(InKernel)));
+#endif  
+			}
+  
+			// Compile the shaders for this shader map now if not deferring and deferred compiles are not enabled globally
+			if (bSynchronousCompile)
+			{
+				TArray<int32> CurrentShaderMapId;
+				CurrentShaderMapId.Add(CompilingId);
+				GComputeKernelShaderCompilationManager.FinishCompilation(*NewContent->FriendlyName, CurrentShaderMapId);
+			}
+		}
+	}
+}
+
+FShader* FNNERuntimeIREEShaderMap::ProcessCompilationResultsForSingleJob(FShaderCompileJob& CurrentJob, const FSHAHash& InShaderMapHash)
+{
+	check(CurrentJob.Id == CompilingId);
+
+	GetResourceCode()->AddShaderCompilerOutput(CurrentJob.Output, CurrentJob.Key.ToString(), CurrentJob.Input.GenerateDebugInfo());
+
+	FShader* Shader = nullptr;
+
+	const FNNERuntimeIREEShaderType* NNERuntimeIREEShaderType = CurrentJob.Key.ShaderType->GetNNERuntimeIREEShaderType();
+	check(NNERuntimeIREEShaderType);
+
+	Shader = NNERuntimeIREEShaderType->FinishCompileShader(InShaderMapHash, CurrentJob, GetContent()->FriendlyName);
+
+	bCompiledSuccessfully = CurrentJob.bSucceeded;
+
+	FNNERuntimeIREEShader* NNERuntimeIREEShader = static_cast<FNNERuntimeIREEShader*>(Shader);
+	check(NNERuntimeIREEShader);
+
+	check(!GetContent()->HasShader(NNERuntimeIREEShaderType, CurrentJob.Key.PermutationId));
+	return GetMutableContent()->FindOrAddShader(NNERuntimeIREEShaderType->GetHashedName(), CurrentJob.Key.PermutationId, Shader);
+}
+
+bool FNNERuntimeIREEShaderMap::ProcessCompilationResults(const TArray<FShaderCommonCompileJobPtr>& InCompilationResults, int32& InOutJobIndex, float& InOutTimeBudget)
+{
+	check(InOutJobIndex < InCompilationResults.Num());
+
+	double StartTime = FPlatformTime::Seconds();
+
+	FSHAHash ShaderMapHash;
+	GetContent()->ShaderMapId.GetKernelHash(ShaderMapHash);
+
+	do
+	{
+		ProcessCompilationResultsForSingleJob(static_cast<FShaderCompileJob&>(*InCompilationResults[InOutJobIndex].GetReference()), ShaderMapHash);
+
+		InOutJobIndex++;
+		
+		double NewStartTime = FPlatformTime::Seconds();
+		InOutTimeBudget -= NewStartTime - StartTime;
+		StartTime = NewStartTime;
+	}
+	while ((InOutTimeBudget > 0.0f) && (InOutJobIndex < InCompilationResults.Num()));
+
+	if (InOutJobIndex == InCompilationResults.Num())
+	{
+		FinalizeContent();
+
+		// SaveToDerivedDataCache();
+		// The shader map can now be used on the rendering thread
+		bCompilationFinalized = true;
+		return true;
+	}
+
+	return false;
+}
+
+#endif // WITH_EDITOR
+
+bool FNNERuntimeIREEShaderMap::IsComputeKernelShaderComplete(const FNNERuntimeIREEResource* InKernel, const FNNERuntimeIREEShaderType* InShaderType, bool bSilent)
+{
+	// If we should cache this kernel, it's incomplete if the shader is missing
+	if (ShouldCacheNNERuntimeIREEShader(InShaderType, GetShaderPlatform(), InKernel))
+	{
+		for (int32 PermutationId = 0; PermutationId < InKernel->GetNumPermutations(); ++PermutationId)
+		{
+			if (!GetContent()->HasShader((FShaderType*)InShaderType, PermutationId))
+			{
+				if (!bSilent)
+				{
+					UE_LOG(LogNNERuntimeIREEShader, Warning, TEXT("Incomplete shader %s, missing FNNERuntimeIREEShader %s."), *InKernel->GetFriendlyName(), InShaderType->GetName());
+				}
+				return false;
+			}
+		}
+	}
+
+	return true;
+}
+
+bool FNNERuntimeIREEShaderMap::IsComplete(const FNNERuntimeIREEResource* InKernel, bool bSilent)
+{
+	// Make sure we are operating on a referenced shader map or the below Find will cause this shader map to be deleted,
+	// Since it creates a temporary ref counted pointer.
+	check(NumRefs > 0);
+
+#if WITH_EDITOR
+	const TArray<FNNERuntimeIREEResource*>* CorrespondingKernels = FNNERuntimeIREEShaderMap::ShaderMapsBeingCompiled.Find(this);
+	if (CorrespondingKernels)
+	{
+		check(!bCompilationFinalized);
+		return false;
+	}
+#endif
+
+	// Iterate over all shader types.
+	for(TLinkedList<FShaderType*>::TIterator ShaderTypeIt(FShaderType::GetTypeList()); ShaderTypeIt; ShaderTypeIt.Next())
+	{
+		// Find this shader type in the kernel's shader map.
+		const FNNERuntimeIREEShaderType* ShaderType = ShaderTypeIt->GetNNERuntimeIREEShaderType();
+		if (ShaderType && !IsComputeKernelShaderComplete(InKernel, ShaderType, bSilent))
+		{
+			return false;
+		}
+	}
+
+	return true;
+}
+
+void FNNERuntimeIREEShaderMap::GetShaderList(TMap<FShaderId, TShaderRef<FShader>>& OutShaders) const
+{
+	GetContent()->GetShaderList(*this, FSHAHash(), OutShaders);
+}
+
+void FNNERuntimeIREEShaderMap::GetShaderList(TMap<FHashedName, TShaderRef<FShader>>& OutShaders) const
+{
+	GetContent()->GetShaderList(*this, OutShaders);
+}
+
+void FNNERuntimeIREEShaderMap::GetShaderPipelineList(TArray<FShaderPipelineRef>& OutShaderPipelines) const
+{
+	GetContent()->GetShaderPipelineList(*this, OutShaderPipelines, FShaderPipeline::EAll);
+}
+
+/**
+ * Registers a ComputeKernel shader map in the global map.
+ */
+void FNNERuntimeIREEShaderMap::Register(EShaderPlatform InShaderPlatform)
+{
+	if (!bRegistered)
+	{
+		INC_DWORD_STAT(STAT_Shaders_NumShaderMaps);
+	}
+
+	GIdToKernelShaderMap[GetShaderPlatform()].Add(GetContent()->ShaderMapId,this);
+	bRegistered = true;
+}
+
+void FNNERuntimeIREEShaderMap::AddRef()
+{
+	check(!bDeletedThroughDeferredCleanup);
+	FPlatformAtomics::InterlockedIncrement(&NumRefs);
+}
+
+void FNNERuntimeIREEShaderMap::Release()
+{
+	check(NumRefs > 0);
+	if(FPlatformAtomics::InterlockedDecrement(&NumRefs) == 0)
+	{
+		if (bRegistered)
+		{
+			DEC_DWORD_STAT(STAT_Shaders_NumShaderMaps);
+
+			GIdToKernelShaderMap[GetShaderPlatform()].Remove(GetContent()->ShaderMapId);
+			bRegistered = false;
+		}
+
+		check(!bDeletedThroughDeferredCleanup);
+		bDeletedThroughDeferredCleanup = true;
+		BeginCleanup(this);
+	}
+}
+
+FNNERuntimeIREEShaderMap::FNNERuntimeIREEShaderMap() :
+	CompilingId(1),
+	NumRefs(0),
+	bDeletedThroughDeferredCleanup(false),
+	bRegistered(false),
+	bCompilationFinalized(true),
+	bCompiledSuccessfully(true)/*,
+	bIsPersistent(true) */
+{
+	checkSlow(IsInGameThread() || IsAsyncLoading());
+	AllKernelShaderMaps.Add(this);
+}
+
+FNNERuntimeIREEShaderMap::~FNNERuntimeIREEShaderMap()
+{ 
+	checkSlow(IsInGameThread() || IsAsyncLoading());
+	check(bDeletedThroughDeferredCleanup);
+	check(!bRegistered);
+	AllKernelShaderMaps.RemoveSwap(this);
+}
+
+bool FNNERuntimeIREEShaderMap::Serialize(FArchive& Ar)
+{
+	// Note: This is saved to the DDC, not into packages (except when cooked)
+	// Backwards compatibility therefore will not work based on the version of Ar
+	// Instead, just bump COMPUTEKERNEL_DERIVEDDATA_VER
+	FShaderSerializeContext Ctx(Ar);
+	return Super::Serialize(Ctx);
+}
\ No newline at end of file
diff --git a/Engine/Plugins/Experimental/NNERuntimeIREE/Source/NNERuntimeIREEShader/Private/NNERuntimeIREEShaderModule.cpp b/Engine/Plugins/Experimental/NNERuntimeIREE/Source/NNERuntimeIREEShader/Private/NNERuntimeIREEShaderModule.cpp
new file mode 100644
--- /dev/null
+++ b/Engine/Plugins/Experimental/NNERuntimeIREE/Source/NNERuntimeIREEShader/Private/NNERuntimeIREEShaderModule.cpp
@@ -0,0 +1,33 @@
+// Copyright Epic Games, Inc. All Rights Reserved.
+
+#include "NNERuntimeIREEShaderModule.h"
+
+#include "Interfaces/IPluginManager.h"
+#include "Misc/CoreDelegates.h"
+#include "Misc/Paths.h"
+#include "NNERuntimeIREEShaderLog.h"
+#include "ShaderCore.h"
+
+DEFINE_LOG_CATEGORY(LogNNERuntimeIREEShader);
+
+void FNNERuntimeIREEShaderModule::StartupModule()
+{
+	const TSharedPtr<IPlugin> Plugin = IPluginManager::Get().FindPlugin(TEXT("NNERuntimeIREE"));
+	if (Plugin.IsValid())
+	{
+		const FString ShadersDir = FPaths::Combine(Plugin->GetBaseDir(), TEXT("Shaders"));
+
+		AddShaderSourceDirectoryMapping(TEXT("/Plugin/NNERuntimeIREEShader"), ShadersDir);
+	}
+	else
+	{
+		UE_LOG(LogNNERuntimeIREEShader, Warning, TEXT("Shaders directory not added. Failed to find NNERuntimeIREE plugin"));
+	}
+}
+
+void FNNERuntimeIREEShaderModule::ShutdownModule()
+{
+
+}
+
+IMPLEMENT_MODULE(FNNERuntimeIREEShaderModule, NNERuntimeIREEShader)
\ No newline at end of file
diff --git a/Engine/Plugins/Experimental/NNERuntimeIREE/Source/NNERuntimeIREEShader/Private/NNERuntimeIREEShaderModule.h b/Engine/Plugins/Experimental/NNERuntimeIREE/Source/NNERuntimeIREEShader/Private/NNERuntimeIREEShaderModule.h
new file mode 100644
--- /dev/null
+++ b/Engine/Plugins/Experimental/NNERuntimeIREE/Source/NNERuntimeIREEShader/Private/NNERuntimeIREEShaderModule.h
@@ -0,0 +1,13 @@
+// Copyright Epic Games, Inc. All Rights Reserved.
+
+#pragma once
+
+#include "Modules/ModuleManager.h"
+
+class FNNERuntimeIREEShaderModule : public IModuleInterface
+{
+public:
+	/** IModuleInterface implementation */
+	virtual void StartupModule() override;
+	virtual void ShutdownModule() override;
+};
diff --git a/Engine/Plugins/Experimental/NNERuntimeIREE/Source/NNERuntimeIREEShader/Private/NNERuntimeIREEShaderShared.cpp b/Engine/Plugins/Experimental/NNERuntimeIREE/Source/NNERuntimeIREEShader/Private/NNERuntimeIREEShaderShared.cpp
new file mode 100644
--- /dev/null
+++ b/Engine/Plugins/Experimental/NNERuntimeIREE/Source/NNERuntimeIREEShader/Private/NNERuntimeIREEShaderShared.cpp
@@ -0,0 +1,343 @@
+// Copyright Epic Games, Inc. All Rights Reserved.
+
+#include "NNERuntimeIREEShaderShared.h"
+
+#include "DataDrivenShaderPlatformInfo.h"
+#include "Interfaces/ITargetPlatform.h"
+#include "NNERuntimeIREEShader.h"
+#include "NNERuntimeIREEShaderCompilationManager.h"
+#include "NNERuntimeIREEShaderLog.h"
+#include "RendererInterface.h"
+#include "RHIShaderFormatDefinitions.inl"
+#include "ShaderCompiler.h"
+#include "ShaderParameterMetadataBuilder.h"
+#include "Stats/StatsMisc.h"
+#include "TextureResource.h"
+#include "UObject/CoreObjectVersion.h"
+#include "UObject/UObjectHash.h"
+#include "UObject/UObjectIterator.h"
+
+IMPLEMENT_TYPE_LAYOUT(FNNERuntimeIREECompilationOutput);
+IMPLEMENT_TYPE_LAYOUT(FNNERuntimeIREEShaderMapId);
+IMPLEMENT_TYPE_LAYOUT(FNNERuntimeIREEShaderMapContent);
+
+FNNERuntimeIREEResource::FNNERuntimeIREEResource() : bLoadedCookedShaderMapId(false)
+{
+}
+
+FNNERuntimeIREEResource::~FNNERuntimeIREEResource()
+{
+
+}
+
+bool FNNERuntimeIREEResource::ShouldCache(EShaderPlatform InPlatform, const FShaderType* InShaderType) const
+{
+	check(InShaderType->GetNNERuntimeIREEShaderType())
+	// return ComputeFramework::IsSupported(InPlatform);
+	return true;
+}
+
+void FNNERuntimeIREEResource::SerializeShaderMap(FArchive& Ar)
+{
+	// if (FPlatformProperties::RequiresCookedData() && !bCooked && Ar.IsLoading())
+	// {
+	// 	UE_LOG(LogNNERuntimeIREEShader, Fatal, TEXT("This platform requires cooked packages, and shaders were not cooked into this kernel %s."), *GetFriendlyName());
+	// }
+
+	if (Ar.IsSaving() /*Ar.IsCooking()*/)
+	{
+#if WITH_EDITOR
+		FinishCompilation();
+
+		bool bValid = GameThreadShaderMap != nullptr && GameThreadShaderMap->CompiledSuccessfully();
+		Ar << bValid;
+
+		if (bValid)
+		{
+			GameThreadShaderMap->AssociateWithAsset(AssetPath);
+
+			GameThreadShaderMap->Serialize(Ar);
+		}
+#endif
+	}
+	else
+	{
+		bool bValid = false;
+		Ar << bValid;
+
+		if (bValid)
+		{
+			TRefCountPtr<FNNERuntimeIREEShaderMap> LoadedShaderMap = new FNNERuntimeIREEShaderMap();
+			bool bSuccessfullyLoaded = LoadedShaderMap->Serialize(Ar);
+
+			// Toss the loaded shader data if this is a server only instance or if it's for a different RHI than the current one.
+			// todo[CF] Don't cook it in the first place
+			if (bSuccessfullyLoaded && FApp::CanEverRender())
+			{
+				GameThreadShaderMap /*= RenderingThreadShaderMap */= LoadedShaderMap;
+				GameThreadShaderMap->GetResource()->SetOwnerName(GetOwnerFName());
+			}
+		}
+	}
+}
+
+void FNNERuntimeIREEResource::RemoveOutstandingCompileId(const int32 InOldOutstandingCompileShaderMapId)
+{
+	if (0 <= OutstandingCompileShaderMapIds.Remove(InOldOutstandingCompileShaderMapId))
+	{
+		UE_LOG(LogNNERuntimeIREEShader, Verbose, TEXT("RemoveOutstandingCompileId %p %d"), this, InOldOutstandingCompileShaderMapId);
+	}
+}
+
+void FNNERuntimeIREEResource::NotifyCompilationFinished(FString const& ResultMessage)
+{
+	UE_LOG(LogNNERuntimeIREEShader, Log, TEXT("%s"), *ResultMessage);
+	FNNERuntimeIREEShaderCompileMessage Message;
+	Message.Type = FNNERuntimeIREEShaderCompileMessage::EMessageType::Info;
+	Message.Text = ResultMessage;
+	CompilationResults.Messages.Add(Message);
+	// OnCompilationCompleteDelegate.ExecuteIfBound(this);
+}
+
+#if WITH_EDITOR
+void FNNERuntimeIREEResource::FinishCompilation()
+{
+	TArray<int32> ShaderMapIdsToFinish;
+	GetShaderMapIDsWithUnfinishedCompilation(ShaderMapIdsToFinish);
+
+	if (ShaderMapIdsToFinish.Num() > 0)
+	{
+		for (int32 i = 0; i < ShaderMapIdsToFinish.Num(); i++)
+		{
+			UE_LOG(LogNNERuntimeIREEShader, Verbose, TEXT("FinishCompilation()[%d] %s id %d!"), i, *GetFriendlyName(), ShaderMapIdsToFinish[i]);
+		}
+	
+		// Block until the shader maps that we will save have finished being compiled
+		GComputeKernelShaderCompilationManager.FinishCompilation(*GetFriendlyName(), ShaderMapIdsToFinish);
+
+		// Shouldn't have anything left to do...
+		TArray<int32> ShaderMapIdsToFinish2;
+		GetShaderMapIDsWithUnfinishedCompilation(ShaderMapIdsToFinish2);
+		ensure(ShaderMapIdsToFinish2.Num() == 0);
+	}
+}
+#endif // WITH_EDITOR
+
+void FNNERuntimeIREEResource::GetDependentShaderTypes(EShaderPlatform InPlatform, TArray<FShaderType*>& OutShaderTypes) const
+{
+	for (TLinkedList<FShaderType*>::TIterator ShaderTypeIt(FShaderType::GetTypeList()); ShaderTypeIt; ShaderTypeIt.Next())
+	{
+		FNNERuntimeIREEShaderType* ShaderType = ShaderTypeIt->GetNNERuntimeIREEShaderType();
+
+		if (ShaderType && ShaderType->ShouldCache(InPlatform, this) && ShouldCache(InPlatform, ShaderType) )
+		{
+			OutShaderTypes.Add(ShaderType);
+		}
+	}
+
+	OutShaderTypes.Sort(FCompareShaderTypes());
+}
+
+void FNNERuntimeIREEResource::GetShaderMapId(EShaderPlatform InPlatform, const ITargetPlatform* TargetPlatform, FNNERuntimeIREEShaderMapId& OutId) const
+{
+	if (bLoadedCookedShaderMapId)
+	{
+		OutId = CookedShaderMapId;
+	}
+	else
+	{
+		TArray<FShaderType*> ShaderTypes;
+		GetDependentShaderTypes(InPlatform, ShaderTypes);
+
+		OutId.FeatureLevel = GetFeatureLevel();
+		OutId.ShaderCodeHash = ShaderCodeHash;
+#if WITH_EDITOR
+		OutId.SetShaderDependencies(ShaderTypes, InPlatform);
+		if (TargetPlatform)
+		{
+			OutId.LayoutParams.InitializeForPlatform(TargetPlatform->IniPlatformName(), TargetPlatform->HasEditorOnlyData());
+		}
+		else
+		{
+			OutId.LayoutParams.InitializeForCurrent();
+		}
+#else
+		if (TargetPlatform != nullptr)
+		{
+			UE_LOG(LogNNERuntimeIREEShader, Error, TEXT("FComputeKernelResource::GetShaderMapId: TargetPlatform is not null, but a cooked executable cannot target platforms other than its own."));
+		}
+		OutId.LayoutParams.InitializeForCurrent();
+#endif
+	}
+}
+
+bool FNNERuntimeIREEResource::CacheShaders(EShaderPlatform InPlatform, const ITargetPlatform* TargetPlatform, bool bApplyCompletedShaderMapForRendering, bool bSynchronous)
+{
+	FNNERuntimeIREEShaderMapId ShaderMapId;
+	GetShaderMapId(InPlatform, TargetPlatform, ShaderMapId);
+	return CacheShaders(ShaderMapId, InPlatform, bApplyCompletedShaderMapForRendering, bSynchronous);
+}
+
+bool FNNERuntimeIREEResource::CacheShaders(const FNNERuntimeIREEShaderMapId& InShaderMapId, EShaderPlatform InPlatform, bool bApplyCompletedShaderMapForRendering, bool bSynchronous)
+{
+	// Find the kernel's cached shader map.
+	GameThreadShaderMap = FNNERuntimeIREEShaderMap::FindId(InShaderMapId, InPlatform);
+	if (GameThreadShaderMap /*&& GameThreadShaderMap->IsComplete(this, false)*/)
+	{
+		return true;
+	}
+
+	bool bSucceeded = false;
+
+#if WITH_EDITOR
+	// If there's no cached shader map for this kernel compile a new one.
+	// This is just kicking off the compile, GameThreadShaderMap will not be complete yet
+	bSucceeded = BeginCompileShaderMap(InShaderMapId, InPlatform, GameThreadShaderMap, bApplyCompletedShaderMapForRendering, bSynchronous);
+#endif // WITH_EDITOR
+
+	if (!bSucceeded)
+	{
+		GameThreadShaderMap = nullptr;
+	}
+
+	return bSucceeded;
+}
+
+void FNNERuntimeIREEResource::SetupResource(
+	ERHIFeatureLevel::Type InFeatureLevel,
+	FString const& InFriendlyName,
+	FString const& InShaderEntryPoint,
+	FString const& InShaderHashKey,
+	FString const& InShaderSource,
+	// TMap<FString, FString> const& InAdditionalSources,
+	// TMap<FString, FString> const& InGeneratedSources,
+	// TSharedPtr<FComputeKernelDefinitionSet>& InShaderDefinitionSet,
+	// TSharedPtr<FComputeKernelPermutationVector>& InShaderPermutationVector,
+	// TUniquePtr<FShaderParametersMetadataAllocations>& InShaderParameterMetadataAllocations,
+	FShaderParametersMetadata* InShaderParameterMetadata,
+	FName const& InAssetPath
+	)
+{
+	FeatureLevel = InFeatureLevel;
+	FriendlyName = InFriendlyName;
+	ShaderEntryPoint = InShaderEntryPoint;
+	ShaderCodeHash = GetTypeHash(InShaderHashKey);
+	ShaderSource = InShaderSource;
+	// AdditionalSources = InAdditionalSources;
+	// GeneratedSources = InGeneratedSources;
+	// ShaderDefinitionSet = InShaderDefinitionSet;
+	// ShaderPermutationVector = InShaderPermutationVector;
+	// ShaderParameterMetadataAllocations = MoveTemp(InShaderParameterMetadataAllocations);
+	ShaderParameterMetadata = InShaderParameterMetadata;
+	CompilationResults.Messages.Reset();
+	AssetPath = InAssetPath;
+}
+
+TShaderRef<FNNERuntimeIREEShader> FNNERuntimeIREEResource::GetShader(int32 PermutationId) const
+{
+	check(!GIsThreadedRendering || !IsInGameThread());
+
+	return GameThreadShaderMap->GetShader<FNNERuntimeIREEShader>(PermutationId);
+
+	// if (!GIsEditor || RenderingThreadShaderMap)
+	// {
+	// 	return RenderingThreadShaderMap->GetShader<FComputeKernelShader>(PermutationId);
+	// }
+	// return TShaderRef<FComputeKernelShader>();
+}
+
+bool FNNERuntimeIREEResource::IsSame(const FNNERuntimeIREEShaderMapId& InIdentifier) const
+{
+	return
+		InIdentifier.ShaderCodeHash == ShaderCodeHash &&
+		InIdentifier.FeatureLevel == FeatureLevel;
+}
+
+#if WITH_EDITOR
+void FNNERuntimeIREEResource::GetShaderMapIDsWithUnfinishedCompilation(TArray<int32>& OutShaderMapIds)
+{
+	// Build an array of the shader map Id's are not finished compiling.
+	if (GameThreadShaderMap && GameThreadShaderMap.IsValid() && !GameThreadShaderMap->IsCompilationFinalized())
+	{
+		OutShaderMapIds.Add(GameThreadShaderMap->GetCompilingId());
+	}
+	else if (OutstandingCompileShaderMapIds.Num() != 0)
+	{
+		OutShaderMapIds.Append(OutstandingCompileShaderMapIds);
+	}
+}
+
+/**
+ * Compiles this kernel for Platform, storing the result in OutShaderMap
+ *
+ * @param InShaderMapId - the set of static parameters to compile
+ * @param InPlatform - the platform to compile for
+ * @param OutShaderMap - the shader map to compile
+ * @return - true if compile succeeded or was not necessary (shader map for ShaderMapId was found and was complete)
+ */
+bool FNNERuntimeIREEResource::BeginCompileShaderMap(const FNNERuntimeIREEShaderMapId& InShaderMapId, EShaderPlatform InPlatform, TRefCountPtr<FNNERuntimeIREEShaderMap>& OutShaderMap, bool bApplyCompletedShaderMapForRendering, bool bSynchronous)
+{
+	bool bSuccess = false;
+
+	STAT(double ComputeKernelCompileTime = 0);
+
+	SCOPE_SECONDS_COUNTER(ComputeKernelCompileTime);
+
+	TRefCountPtr<FNNERuntimeIREEShaderMap> NewShaderMap = new FNNERuntimeIREEShaderMap();
+
+	// Create a shader compiler environment for the material that will be shared by all jobs from this material
+	TRefCountPtr<FSharedShaderCompilerEnvironment> Environment = new FSharedShaderCompilerEnvironment();
+
+	// Compile the shaders for the kernel.
+	FNNERuntimeIREECompilationOutput CompilationOutput;
+	NewShaderMap->Compile(this, InShaderMapId, Environment, CompilationOutput, InPlatform, bSynchronous, bApplyCompletedShaderMapForRendering);
+
+	if (bSynchronous)
+	{
+		// If this is a synchronous compile, assign the compile result to the output
+		OutShaderMap = NewShaderMap->CompiledSuccessfully() ? NewShaderMap : nullptr;
+	}
+	else
+	{
+		UE_LOG(LogNNERuntimeIREEShader, Verbose, TEXT("BeginCompileShaderMap AddUnique %p %d"), this, NewShaderMap->GetCompilingId());
+		unimplemented();
+		// OutstandingCompileShaderMapIds.AddUnique(NewShaderMap->GetCompilingId());
+		
+		// Async compile, use nullptr to detect it if used
+		OutShaderMap = nullptr;
+	}
+
+	return true;
+}
+
+void FNNERuntimeIREEShaderMapId::SetShaderDependencies(const TArray<FShaderType*>& InShaderTypes, EShaderPlatform InShaderPlatform)
+{
+	if (!FPlatformProperties::RequiresCookedData())
+	{
+		for (FShaderType* ShaderType : InShaderTypes)
+		{
+			if (ShaderType != nullptr)
+			{
+				FShaderTypeDependency Dependency;
+				Dependency.ShaderTypeName = ShaderType->GetHashedName();
+				Dependency.SourceHash = ShaderType->GetSourceHash(InShaderPlatform);
+				ShaderTypeDependencies.Add(Dependency);
+			}
+		}
+	}
+}
+
+#endif // WITH_EDITOR
+
+bool FNNERuntimeIREEShaderMapId::ContainsShaderType(const FShaderType* ShaderType) const
+{
+	for (int32 TypeIndex = 0; TypeIndex < ShaderTypeDependencies.Num(); TypeIndex++)
+	{
+		if (ShaderTypeDependencies[TypeIndex].ShaderTypeName == ShaderType->GetHashedName())
+		{
+			return true;
+		}
+	}
+
+	return false;
+}
\ No newline at end of file
