[NNE] NNERuntimeIREE added missing include.

#rnx

diff --git a/Engine/Plugins/Experimental/NNERuntimeIREE/Source/IREEDriverRDG/Internal/IREEDriverRDG.h b/Engine/Plugins/Experimental/NNERuntimeIREE/Source/IREEDriverRDG/Internal/IREEDriverRDG.h
--- a/Engine/Plugins/Experimental/NNERuntimeIREE/Source/IREEDriverRDG/Internal/IREEDriverRDG.h
+++ b/Engine/Plugins/Experimental/NNERuntimeIREE/Source/IREEDriverRDG/Internal/IREEDriverRDG.h
@@ -4,6 +4,7 @@
 
 #ifdef WITH_IREE_DRIVER_RDG
 
+#include "Containers/UnrealString.h"
 #include "HAL/Platform.h"
 
 #if PLATFORM_MICROSOFT
