[NNERuntimeIREECpu] Memory allocation and buffer copy improvements

Problem:  - Currently the iree runtime uses the default iree heap allocator for host allocations. This bypasses the UE allocator therefore doesn't show up in Unreal Insight.
- All input buffers are copied during each call.

Solution: - Create our own heap allocator that calls UEs memory allocation functions and tags there use.
- Use a different library call to create a buffer view that doesn't perform a copy.
- Add caches in order to reduce the number of allocations needed for each inference call.

Test :
* Win64 Editor debug : NNE standard API unit test (from a different CL) + NNEQA IREE test from session frontend + NNEQA smoketests (including IREE test) + IREE plugin test
* Win64 Standalone : NNE standard API unit test (from a different CL) + NNEQA IREE test + NNEQA smoketests (including IREE test) + style transfer
* Consoles: NNEQA smoketests (including IREE test) + style transfer

#rnx

diff --git a/Engine/Plugins/Experimental/NNERuntimeIREE/Source/NNERuntimeIREE/Private/NNERuntimeIREEModel.cpp b/Engine/Plugins/Experimental/NNERuntimeIREE/Source/NNERuntimeIREE/Private/NNERuntimeIREEModel.cpp
--- a/Engine/Plugins/Experimental/NNERuntimeIREE/Source/NNERuntimeIREE/Private/NNERuntimeIREEModel.cpp
+++ b/Engine/Plugins/Experimental/NNERuntimeIREE/Source/NNERuntimeIREE/Private/NNERuntimeIREEModel.cpp
@@ -17,6 +17,7 @@
 #include "NNEStatus.h"
 #include "NNETensor.h"
 #include "RenderGraphUtils.h"
+#include "HAL/UnrealMemory.h"
 #include "Serialization/Archive.h"
 
 #if PLATFORM_MICROSOFT
@@ -26,10 +27,12 @@
 THIRD_PARTY_INCLUDES_START
 #include "iree/hal/drivers/local_sync/sync_device.h"
 #include "iree/hal/local/loaders/static_library_loader.h"
+#include "iree/modules/hal/types.h"
 #include "iree/runtime/call.h"
 #include "iree/runtime/instance.h"
 #include "iree/runtime/session.h"
 #include "iree/vm/bytecode/module.h"
+#include "iree/vm/list.h"
 THIRD_PARTY_INCLUDES_END
 #if PLATFORM_MICROSOFT
 #include "Microsoft/HideMicrosoftPlatformAtomics.h"
@@ -40,6 +43,7 @@ BEGIN_SHADER_PARAMETER_STRUCT(FIREESessionRDGParameters, )
 	RDG_BUFFER_ACCESS_ARRAY(InputBuffers)
 	RDG_BUFFER_ACCESS_ARRAY(OutputBuffers)
 END_SHADER_PARAMETER_STRUCT()
+LLM_DEFINE_TAG(NNERuntimeIREE_Cpu);
 
 DECLARE_GPU_STAT_NAMED(FNNERuntimeIREERDG, TEXT("EnqueueRDG"));
 
@@ -119,6 +123,56 @@ namespace UE::NNERuntimeIREE
 			FMemory::Free(ErrorString);
 		}
 
+		iree_status_t HostAllocatorImplAlloc(iree_allocator_command_t Command, const iree_allocator_alloc_params_t* Params, void** InoutPtr)
+		{
+			check(Params);
+			check(InoutPtr);
+			switch (Command)
+			{
+				case IREE_ALLOCATOR_COMMAND_MALLOC:
+					*InoutPtr = FMemory::Malloc(Params->byte_length, IREE_HAL_HEAP_BUFFER_ALIGNMENT);
+					break;
+				case IREE_ALLOCATOR_COMMAND_CALLOC:
+					*InoutPtr = FMemory::MallocZeroed(Params->byte_length, IREE_HAL_HEAP_BUFFER_ALIGNMENT);
+					break;
+				case IREE_ALLOCATOR_COMMAND_REALLOC:
+					*InoutPtr = FMemory::Realloc(*InoutPtr, Params->byte_length, IREE_HAL_HEAP_BUFFER_ALIGNMENT);
+					break;
+				default:
+					check(false);
+			}
+			if (!*InoutPtr)
+			{
+				return iree_make_status(IREE_STATUS_RESOURCE_EXHAUSTED, "Host allocator failed during alloc");
+			}
+			return iree_ok_status();
+		}
+
+		iree_status_t HostAllocatorImpl(void *Self, iree_allocator_command_t Command, const void* Params, void** InoutPtr)
+		{
+			LLM_SCOPE_BYTAG(NNERuntimeIREE_Cpu);
+			check(InoutPtr);
+			switch (Command)
+			{
+				case IREE_ALLOCATOR_COMMAND_MALLOC:
+				case IREE_ALLOCATOR_COMMAND_CALLOC:
+				case IREE_ALLOCATOR_COMMAND_REALLOC:
+					return HostAllocatorImplAlloc(Command, static_cast<const iree_allocator_alloc_params_t*>(Params), InoutPtr);
+				case IREE_ALLOCATOR_COMMAND_FREE:
+					FMemory::Free(*InoutPtr);
+					break;
+				default:
+					return iree_make_status(IREE_STATUS_UNIMPLEMENTED, "Host allocator command not recognised");
+			}
+			return iree_ok_status();
+		}
+
+		inline iree_allocator_t MakeHostAllocator()
+		{
+			iree_allocator_t Allocator = { NULL, HostAllocatorImpl };
+			return Allocator;
+		}
+
 		class FInstance
 		{
 		private:
@@ -158,7 +212,7 @@ namespace UE::NNERuntimeIREE
 				iree_runtime_instance_options_use_all_available_drivers(&InstanceOptions);
 
 				iree_runtime_instance_t* Instance = nullptr;
-				Status = iree_runtime_instance_create(&InstanceOptions, iree_allocator_system(), &Instance);
+				Status = iree_runtime_instance_create(&InstanceOptions, MakeHostAllocator(), &Instance);
 				if (!iree_status_is_ok(Status))
 				{
 					Private::PrintIREEError("CPU instance: Failed to create the instance", Status);
@@ -599,22 +653,14 @@ namespace UE::NNERuntimeIREE
 					return iree_hal_device_allocator(Device);
 				}
 
-				bool CreateBufferViewCopy(iree_host_size_t Rank, iree_hal_dim_t* Shape, iree_hal_element_types_t Type, iree_hal_buffer_params_t Params, void* Data, iree_host_size_t DataSizeInBytes, iree_hal_buffer_view_t** BufferView)
+				iree_allocator_t GetHostAllocator()
 				{
-					iree_status_t Status = iree_ok_status();
-					check(iree_status_is_ok(Status));
-
-					Status = iree_hal_buffer_view_allocate_buffer_copy(
-						Device,
-						GetDeviceAllocator(),
-						Rank, Shape,
-						Type, IREE_HAL_ENCODING_TYPE_DENSE_ROW_MAJOR,
-						Params,
-						iree_make_const_byte_span(Data, DataSizeInBytes),
-						BufferView);
+					return iree_hal_device_host_allocator(Device);
+				}
 
-					iree_status_free(Status);
-					return iree_status_is_ok(Status);
+				iree_status_t CopyFromBuffer(iree_hal_buffer_t *Source, void *Target, iree_device_size_t CopySizeInBytes)
+				{
+					return iree_hal_device_transfer_d2h(Device, Source, 0, Target, CopySizeInBytes, IREE_HAL_TRANSFER_BUFFER_FLAG_DEFAULT, iree_infinite_timeout());
 				}
 			};
 			TMap<FString, TWeakPtr<FDevice>> FDevice::Devices;
@@ -630,12 +676,15 @@ namespace UE::NNERuntimeIREE
 				TArray<UE::NNE::FTensorDesc> OutputTensorDescs;
 				TArray<UE::NNE::FTensorShape> InputTensorShapes;
 				TArray<UE::NNE::FTensorShape> OutputTensorShapes;
+				TArray<UE::NNE::FTensorBindingCPU> PreviousInputBindings;
 
 				FSession(TSharedRef<FDevice> InDevice, TSharedRef<UE::NNERuntimeIREE::Private::FModule> InModule, iree_runtime_session_t* InSession, const iree_runtime_call_t& InCall, TConstArrayView<UE::NNE::FTensorDesc> InInputTensorDescs, TConstArrayView<UE::NNE::FTensorDesc> InOutputTensorDescs)
 					: Device(InDevice), Module(InModule), Session(InSession), Call(InCall), InputTensorDescs(InInputTensorDescs), OutputTensorDescs(InOutputTensorDescs)
 				{
 					check(InSession);
 					check(!InputTensorDescs.IsEmpty());
+					PreviousInputBindings.SetNum(InputTensorDescs.Num());
+					iree_vm_list_resize(iree_runtime_call_inputs(&Call), InputTensorDescs.Num());
 				}
 
 			public:
@@ -765,32 +814,25 @@ namespace UE::NNERuntimeIREE
 
 				ERunSyncStatus RunSyncCPU(TConstArrayView<UE::NNE::FTensorBindingCPU> InInputBindings, TConstArrayView<UE::NNE::FTensorBindingCPU> InOutputBindings)
 				{
+					LLM_SCOPE_BYTAG(NNERuntimeIREE_Cpu);
 					SCOPED_NAMED_EVENT_TEXT("NNERuntimeIREE::CPU::RunSync", FColor::Magenta);
-
 					// Verify the model inputs were prepared
 					if (InputTensorShapes.IsEmpty())
 					{
 						UE_LOG(LogNNERuntimeIREE, Error, TEXT("CPU session: Input shapes are not set, please call SetInputTensorShapes."));
 						return ERunSyncStatus::Fail;
 					}
-
 					if (InInputBindings.Num() != InputTensorShapes.Num())
 					{
 						UE_LOG(LogNNERuntimeIREE, Error, TEXT("CPU session: Input bindings need to match input tensor descriptor count (got %d, expected %d)."), InInputBindings.Num(), InputTensorShapes.Num());
 						return ERunSyncStatus::Fail;
 					}
-
 					if (!InOutputBindings.IsEmpty() && InOutputBindings.Num() != OutputTensorDescs.Num())
 					{
 						UE_LOG(LogNNERuntimeIREE, Error, TEXT("CPU session: Output binding can be empty or needs to match output tensor descriptor count (got %d, expected %d)."), InOutputBindings.Num(), OutputTensorDescs.Num());
 						return ERunSyncStatus::Fail;
 					}
-
-					iree_status_t Status = iree_ok_status();
-					check(iree_status_is_ok(Status));
-
-					iree_runtime_call_reset(&Call);
-
+					
 					for (int32 i = 0; i < InInputBindings.Num(); i++)
 					{
 						if (!InInputBindings[i].Data && InInputBindings[i].SizeInBytes != 0)
@@ -798,50 +840,16 @@ namespace UE::NNERuntimeIREE
 							UE_LOG(LogNNERuntimeIREE, Error, TEXT("CPU session: Binding input tensor %d is not set but given size is non-zero %d."), i, InInputBindings[i].SizeInBytes);
 							return ERunSyncStatus::Fail;
 						}
-
 						if (InInputBindings[i].SizeInBytes != (InputTensorShapes[i].Volume() * InputTensorDescs[i].GetElementByteSize()))
 						{
 							UE_LOG(LogNNERuntimeIREE, Error, TEXT("CPU session: Binding input tensor %d size does not match size given by tensor descriptor (got %d, expected %d)."), i, InInputBindings[i].SizeInBytes, InputTensorShapes[i].Volume() * InputTensorDescs[i].GetElementByteSize());
 							return ERunSyncStatus::Fail;
 						}
-
 						if (FMath::Modulo<uint64>((uint64)InInputBindings[i].Data, IREE_HAL_HEAP_BUFFER_ALIGNMENT) != 0)
 						{
 							UE_LOG(LogNNERuntimeIREE, Error, TEXT("CPU session: Input bindings memory need to be aligned with %d bytes"), IREE_HAL_HEAP_BUFFER_ALIGNMENT);
 							return ERunSyncStatus::Fail;
 						}
-
-						iree_hal_buffer_view_t* TempBufferView;
-						iree_hal_buffer_params_t Params = { 0 };
-						Params.type = IREE_HAL_MEMORY_TYPE_DEVICE_LOCAL;
-						Params.usage = IREE_HAL_BUFFER_USAGE_DEFAULT;
-						iree_hal_dim_t Shape[UE::NNE::FTensorShape::MaxRank];
-						for (int32 j = 0; j < InputTensorShapes[i].Rank(); j++)
-						{
-							Shape[j] = InputTensorShapes[i].GetData()[j];
-						}
-						ENNETensorDataType NNEType = InputTensorDescs[i].GetDataType();
-						iree_hal_element_types_t IREEType = UE::NNERuntimeIREE::Private::NNEToIREEType(NNEType);
-
-						if (!Device->CreateBufferViewCopy(InputTensorShapes[i].Rank(), Shape, IREEType, Params, (void*)InInputBindings[i].Data, InInputBindings[i].SizeInBytes, &TempBufferView))
-						{
-							UE::NNERuntimeIREE::Private::PrintIREEError("CPU session: Failed to allocate the buffer view", Status);
-							if (TempBufferView)
-							{
-								iree_hal_buffer_view_release(TempBufferView);
-							}
-							iree_status_free(Status);
-							return ERunSyncStatus::Fail;
-						}
-
-						Status = iree_runtime_call_inputs_push_back_buffer_view(&Call, TempBufferView);
-						iree_hal_buffer_view_release(TempBufferView);
-						if (!iree_status_is_ok(Status))
-						{
-							UE::NNERuntimeIREE::Private::PrintIREEError("CPU session: Failed to push the buffer view to the input list", Status);
-							iree_status_free(Status);
-							return ERunSyncStatus::Fail;
-						}
 					}
 
 					for (int32 i = 0; i < InOutputBindings.Num(); i++)
@@ -853,72 +861,156 @@ namespace UE::NNERuntimeIREE
 						}
 					}
 
-					Status = iree_runtime_call_invoke(&Call, 0);
+					iree_status_t Status = RunSyncCPUInternal(InInputBindings, InOutputBindings);
+
 					if (!iree_status_is_ok(Status))
 					{
-						UE::NNERuntimeIREE::Private::PrintIREEError("CPU session: Failed to call the model function", Status);
+						UE::NNERuntimeIREE::Private::PrintIREEError("CPU session: RunSyncCPU failed", Status);
 						iree_status_free(Status);
 						return ERunSyncStatus::Fail;
 					}
+					return ERunSyncStatus::Ok;
+				}
 
-					// Set output shapes
-					OutputTensorShapes.Reset();
-					TArray<iree_hal_buffer_view_t*> BufferViews;
-					iree_hal_buffer_view_t* BufferView = nullptr;
-					Status = iree_runtime_call_outputs_pop_front_buffer_view(&Call, &BufferView);
-					while (iree_status_is_ok(Status))
+			private:
+				bool ShapeEqualDims(const UE::NNE::FTensorShape &Shape, const iree_hal_dim_t *Dims, iree_host_size_t Rank)
+				{
+					if (Shape.Rank() != Rank)
 					{
-						iree_host_size_t Rank = iree_hal_buffer_view_shape_rank(BufferView);
-						const iree_hal_dim_t* Dims = iree_hal_buffer_view_shape_dims(BufferView);
-						uint32 Shape[UE::NNE::FTensorShape::MaxRank];
-						for (int32 i = 0; i < FMath::Min((int32)Rank, UE::NNE::FTensorShape::MaxRank); i++)
+						return false;
+					}
+					TConstArrayView<uint32> ShapeData = Shape.GetData();
+					for (int i = 0; i < Rank; i++)
+					{
+						if (ShapeData[i] != Dims[i])
+						{
+							return false;
+						}
+					}
+					return true;
+				}
+
+				iree_status_t RunSyncCPUInternal(TConstArrayView<UE::NNE::FTensorBindingCPU> InInputBindings, TConstArrayView<UE::NNE::FTensorBindingCPU> InOutputBindings)
+				{
+					iree_hal_allocator_t *DeviceAllocator = Device->GetDeviceAllocator();
+					iree_allocator_t HostAllocator = Device->GetHostAllocator();
+					iree_vm_list_t *Inputs = iree_runtime_call_inputs(&Call);
+					iree_vm_list_t *Outputs = iree_runtime_call_outputs(&Call);
+					check(PreviousInputBindings.Num() == InInputBindings.Num());
+					check(iree_vm_list_size(Inputs) == InInputBindings.Num());
+					for (int32 i = 0; i < InInputBindings.Num(); i++)
+					{
+						//Create a new input buffer view if the input binding is different
+						//thant that of the previous call
+						if (PreviousInputBindings[i].Data != InInputBindings[i].Data ||
+							PreviousInputBindings[i].SizeInBytes != InInputBindings[i].SizeInBytes)
 						{
-							Shape[i] = (uint32)Dims[i];
+							//Only set previous input binding once we successfully added 
+							//the input binding to the Inputs list
+							PreviousInputBindings[i] = { };
+							// These flags are added when creating a heap buffer. Source: iree\hal\allocator_heap.c:162
+							const iree_hal_memory_type_t HEAP_MEMORY_TYPE = IREE_HAL_MEMORY_TYPE_HOST_VISIBLE;
+							const iree_hal_buffer_usage_t HEAP_BUFFER_USAGE =
+								IREE_HAL_BUFFER_USAGE_MAPPING_SCOPED |
+								IREE_HAL_BUFFER_USAGE_MAPPING_PERSISTENT |
+								IREE_HAL_BUFFER_USAGE_MAPPING_ACCESS_RANDOM |
+								IREE_HAL_BUFFER_USAGE_TRANSFER;
+
+							iree_hal_buffer_t* Buffer = nullptr;
+							IREE_RETURN_IF_ERROR(
+								iree_hal_heap_buffer_wrap(DeviceAllocator,
+														  IREE_HAL_MEMORY_TYPE_DEVICE_LOCAL | HEAP_MEMORY_TYPE,
+														  IREE_HAL_MEMORY_ACCESS_READ,
+														  IREE_HAL_BUFFER_USAGE_DEFAULT | HEAP_BUFFER_USAGE,
+														  InInputBindings[i].SizeInBytes,
+														  iree_make_byte_span(InInputBindings[i].Data, InInputBindings[i].SizeInBytes),
+														  iree_hal_buffer_release_callback_null(),
+														  &Buffer),
+								"heap buffer wrap failed");
+
+							iree_hal_dim_t Shape[UE::NNE::FTensorShape::MaxRank];
+							for (int32 j = 0; j < InputTensorShapes[i].Rank(); j++)
+							{
+								Shape[j] = InputTensorShapes[i].GetData()[j];
+							}
+							ENNETensorDataType NNEType = InputTensorDescs[i].GetDataType();
+							iree_hal_element_types_t IREEType = UE::NNERuntimeIREE::Private::NNEToIREEType(NNEType);
+							iree_hal_buffer_view_t* BufferView = nullptr;
+							IREE_RETURN_AND_EVAL_IF_ERROR(iree_hal_buffer_release(Buffer),
+								iree_hal_buffer_view_create(Buffer, InputTensorShapes[i].Rank(), Shape,
+															IREEType, IREE_HAL_ENCODING_TYPE_DENSE_ROW_MAJOR,
+															HostAllocator, &BufferView),
+								"creation of input BufferView failed");
+							//Now that BufferView also has ownership of the Buffer we can release our ownership
+							iree_hal_buffer_release(Buffer);
+
+							iree_vm_ref_t BufferViewRef = {0};
+							//Transfers ownership from BufferView to BufferViewRef
+							IREE_RETURN_AND_EVAL_IF_ERROR(iree_hal_buffer_view_release(BufferView),
+								iree_vm_ref_wrap_assign(BufferView, iree_hal_buffer_view_type(), &BufferViewRef),
+								"wrap assignment of BufferView failed");
+							//Releases previous element at i and transfers ownership from BufferViewRef to Inputs
+  							IREE_RETURN_AND_EVAL_IF_ERROR(iree_vm_ref_release(&BufferViewRef),
+								iree_vm_list_set_ref_move(Inputs, i, &BufferViewRef),
+								"set BufferView to input list failed");
+							PreviousInputBindings[i] = InInputBindings[i];
 						}
-						OutputTensorShapes.Add(UE::NNE::FTensorShape::Make(TConstArrayView<uint32>(Shape, FMath::Min((int32)Rank, UE::NNE::FTensorShape::MaxRank))));
+					}
+					//Releases all elements and sets count to 0. Doesn't release the lists memory
+					//This is needed in case the previous call failed to pop all elements
+					iree_vm_list_clear(Outputs);
+
+					IREE_RETURN_IF_ERROR(
+						iree_runtime_call_invoke(&Call, 0),
+						"UE::NNERuntimeIREE::CPU::Private::FSession failed to call the model function");
 
-						BufferViews.Add(BufferView);
-						Status = iree_runtime_call_outputs_pop_front_buffer_view(&Call, &BufferView);
+					check(iree_vm_list_size(Outputs) == OutputTensorDescs.Num());
+					check(OutputTensorShapes.Num() == 0 || OutputTensorShapes.Num() == OutputTensorDescs.Num());
+					if (OutputTensorShapes.Num() == 0 && OutputTensorDescs.Num() > 0)
+					{
+						OutputTensorShapes.SetNum(OutputTensorDescs.Num());
 					}
+					check(OutputTensorShapes.Num() == OutputTensorDescs.Num());
+					check(InOutputBindings.IsEmpty() || InOutputBindings.Num() == OutputTensorDescs.Num());
 
-					// Copy result back only if we have the required buffers available and big enought
-					ERunSyncStatus Result = ERunSyncStatus::Ok;
-					if (!InOutputBindings.IsEmpty())
+					for (int32 i = 0; i < OutputTensorDescs.Num(); i++)
 					{
-						check(OutputTensorShapes.Num() == InOutputBindings.Num())
-						for (int32 i = 0; i < InOutputBindings.Num(); i++)
+						iree_hal_buffer_view_t* BufferView = nullptr;
+						IREE_RETURN_IF_ERROR(
+							iree_runtime_call_outputs_pop_front_buffer_view(&Call, &BufferView),
+							"failed to get output at index %d", i);
+						iree_host_size_t Rank = iree_hal_buffer_view_shape_rank(BufferView);
+						const iree_hal_dim_t* Dims = iree_hal_buffer_view_shape_dims(BufferView);
+						if (!ShapeEqualDims(OutputTensorShapes[i], Dims, Rank))
 						{
-							int32 DataSizeInBytes = iree_hal_buffer_view_element_size(BufferViews[i]) * iree_hal_buffer_view_element_count(BufferViews[i]);
-							if (InOutputBindings[i].Data && InOutputBindings[i].SizeInBytes <= DataSizeInBytes)
+							uint32 Shape[UE::NNE::FTensorShape::MaxRank];
+							int32 ShapeRank = FMath::Min((int32)Rank, UE::NNE::FTensorShape::MaxRank);
+							for (int32 ShapeIndex = 0; ShapeIndex < ShapeRank; ShapeIndex++)
 							{
-								iree_hal_buffer_t* Buffer = iree_hal_buffer_view_buffer(BufferViews[i]);
+								Shape[ShapeIndex] = (uint32)Dims[ShapeIndex];
+							}
+							OutputTensorShapes[i] = UE::NNE::FTensorShape::Make(TConstArrayView<uint32>(Shape, ShapeRank));
+						}
+						if (!InOutputBindings.IsEmpty() && InOutputBindings[i].Data)
+						{
+							int32 DataSizeInBytes = iree_hal_buffer_view_byte_length(BufferView); 
+							if (InOutputBindings[i].SizeInBytes <= DataSizeInBytes)
+							{
+								iree_hal_buffer_t* Buffer = iree_hal_buffer_view_buffer(BufferView);
 								if (!Buffer)
 								{
-									UE_LOG(LogNNERuntimeIREE, Error, TEXT("CPU session: Failed failed to get the result buffer"));
-									Result = ERunSyncStatus::Fail;
-									break;
+									iree_hal_buffer_view_destroy(BufferView);
+									return iree_make_status(IREE_STATUS_UNKNOWN, "Failed to get the result buffer");
 								}
 
-								iree_hal_buffer_mapping_t BufferMapping;
-								Status = iree_hal_buffer_map_range(Buffer, IREE_HAL_MAPPING_MODE_PERSISTENT, IREE_HAL_MEMORY_ACCESS_READ, 0, DataSizeInBytes, &BufferMapping);
-								if (!iree_status_is_ok(Status))
-								{
-									UE::NNERuntimeIREE::Private::PrintIREEError("CPU session: Failed failed to map the result buffer", Status);
-									Result = ERunSyncStatus::Fail;
-									break;
-								}
-								FMemory::Memcpy(InOutputBindings[i].Data, BufferMapping.contents.data, DataSizeInBytes);
-								iree_hal_buffer_unmap_range(&BufferMapping);
+								IREE_RETURN_AND_EVAL_IF_ERROR(iree_hal_buffer_view_destroy(BufferView),
+									Device->CopyFromBuffer(Buffer, InOutputBindings[i].Data, DataSizeInBytes),
+									"Copy to the output buffer failed");
 							}
 						}
+						iree_hal_buffer_view_destroy(BufferView);
 					}
-
-					for (int32 i = 0; i < InOutputBindings.Num(); i++)
-					{
-						iree_hal_buffer_view_release(BufferViews[i]);
-					}
-					iree_status_free(Status);
-					return Result;
+					return iree_ok_status();
 				}
 			};
 		}
diff --git a/Engine/Plugins/Experimental/NNERuntimeIREE/Source/NNERuntimeIREE/Private/NNERuntimeIREEModel.h b/Engine/Plugins/Experimental/NNERuntimeIREE/Source/NNERuntimeIREE/Private/NNERuntimeIREEModel.h
--- a/Engine/Plugins/Experimental/NNERuntimeIREE/Source/NNERuntimeIREE/Private/NNERuntimeIREEModel.h
+++ b/Engine/Plugins/Experimental/NNERuntimeIREE/Source/NNERuntimeIREE/Private/NNERuntimeIREEModel.h
@@ -8,6 +8,8 @@
 #include "NNERuntimeRDG.h"
 #include "NNERuntimeIREEMetaData.h"
 
+LLM_DECLARE_TAG(NNERuntimeIREE_Cpu);
+
 namespace UE::NNERuntimeIREE
 {
 	namespace Private
