[NNE] NNERuntimeIREERdg fix shader compilation on Mac.

Problem:
- Fill buffer shader does not compile on Mac

Solution:
- remove 'register' keywork in resource binding

#tests Editor, Mac verified successful compilation

#rnx

diff --git a/Engine/Plugins/Experimental/NNERuntimeIREE/Shaders/NNERuntimeIREEShaderFillBuffer.usf b/Engine/Plugins/Experimental/NNERuntimeIREE/Shaders/NNERuntimeIREEShaderFillBuffer.usf
--- a/Engine/Plugins/Experimental/NNERuntimeIREE/Shaders/NNERuntimeIREEShaderFillBuffer.usf
+++ b/Engine/Plugins/Experimental/NNERuntimeIREE/Shaders/NNERuntimeIREEShaderFillBuffer.usf
@@ -2,7 +2,7 @@
 
 #include "/Engine/Public/Platform.ush"
 
-RWByteAddressBuffer TargetBuffer : register(u0);
+RWByteAddressBuffer TargetBuffer;
 
 // Fill parameters
 // x: pattern
