[NNE] NNERuntimeIREERdg fixed remaining shader compilation issues for OIDN models.
- added preprocess steps before shader compilation to modify HLSL code:
* convert unsupported suffix 'll' to 'l'
* convert unsupported type 'min16float' to 'float16_t'

#tests Editor, Win64 and packaged, ran all test models

#rnx

diff --git a/Engine/Plugins/Experimental/NNERuntimeIREE/Source/IREEDriverRDG/Private/IREECompilerRDG.cpp b/Engine/Plugins/Experimental/NNERuntimeIREE/Source/IREEDriverRDG/Private/IREECompilerRDG.cpp
--- a/Engine/Plugins/Experimental/NNERuntimeIREE/Source/IREEDriverRDG/Private/IREECompilerRDG.cpp
+++ b/Engine/Plugins/Experimental/NNERuntimeIREE/Source/IREEDriverRDG/Private/IREECompilerRDG.cpp
@@ -152,6 +152,17 @@ namespace UE::IREE::Compiler
 					return false;
 				}
 
+				if (FDataDrivenShaderPlatformInfo::GetSupportsRealTypes(ShaderPlatform) == ERHIFeatureSupport::RuntimeGuaranteed)
+				{
+					HlslSource.ReplaceInline(TEXT("min16float2"), TEXT("float16_t2"));
+				}
+
+				if (FDataDrivenShaderPlatformInfo::GetIsLanguageSony(ShaderPlatform))
+				{
+					// No support for C99 long long and long double data types.
+					HlslSource.ReplaceInline(TEXT("ull"), TEXT("ul"));
+				}
+
 				const FStaticFeatureLevel FeatureLevel = GetMaxSupportedFeatureLevel(ShaderPlatform);
 
 				const FString HlslEntryPointName = TEXT("main");
diff --git a/Engine/Plugins/Experimental/NNERuntimeIREE/Source/NNERuntimeIREEShader/Internal/NNERuntimeIREEShaderType.h b/Engine/Plugins/Experimental/NNERuntimeIREE/Source/NNERuntimeIREEShader/Internal/NNERuntimeIREEShaderType.h
--- a/Engine/Plugins/Experimental/NNERuntimeIREE/Source/NNERuntimeIREEShader/Internal/NNERuntimeIREEShaderType.h
+++ b/Engine/Plugins/Experimental/NNERuntimeIREE/Source/NNERuntimeIREEShader/Internal/NNERuntimeIREEShaderType.h
@@ -7,6 +7,7 @@
 #include "GlobalShader.h"
 #include "Shader.h"
 #if WITH_EDITOR
+#include "DataDrivenShaderPlatformInfo.h"
 #include "ShaderCompiler.h"
 #endif
 
@@ -151,6 +152,11 @@ protected:
 	void SetupCompileEnvironment(EShaderPlatform InPlatform, const FNNERuntimeIREEResource* Kernel, FShaderCompilerEnvironment& OutEnvironment) const
 	{
 		ModifyCompilationEnvironment(FNNERuntimeIREEShaderPermutationParameters(InPlatform), OutEnvironment);
+
+		if (FDataDrivenShaderPlatformInfo::GetSupportsRealTypes(InPlatform) == ERHIFeatureSupport::RuntimeGuaranteed)
+		{
+			OutEnvironment.CompilerFlags.Add(CFLAG_AllowRealTypes);
+		}
 	}
 #endif // WITH_EDITOR
 };
\ No newline at end of file
