[NNE] NNERuntimeIREERdg completed support for 16-bit float format models on all RHI's except Vulkan.
- Updated IREE compiler (added missing pack/unpack to spirv_cross for special cases)
- Compile IREE shaders for shader platforms that might support 16-bit types
- Only register runtime if in Editor or RHI has min. feature level SM6 and 16-bit type support
- Added more tests/logs to early catch and report failures (shader compilation/serialization etc.)
- Fixed bug where files for staging were not copied when cooking

Tests:
- Editor, Win64 and Mac, all RDG model tests green
- Game, Win64 and Mac, all RDG model tests green
- Game, Consoles, all RDG model tests green

#tests Editor & Standalone, Win64 (D3D12) & Mac, RDG model tests green

#rnx

diff --git a/Engine/Plugins/Experimental/NNERuntimeIREE/Source/IREEDriverRDG/Private/IREECompilerRDG.cpp b/Engine/Plugins/Experimental/NNERuntimeIREE/Source/IREEDriverRDG/Private/IREECompilerRDG.cpp
--- a/Engine/Plugins/Experimental/NNERuntimeIREE/Source/IREEDriverRDG/Private/IREECompilerRDG.cpp
+++ b/Engine/Plugins/Experimental/NNERuntimeIREE/Source/IREEDriverRDG/Private/IREECompilerRDG.cpp
@@ -174,7 +174,8 @@ namespace UE::IREE::Compiler
 
 				FIREEDriverRDGShaderParametersMetadata::StaticStruct()->SerializeBin(Writer, &IREEShaderParametersMetadata);
 
-				KernelResource->SerializeShaderMap(Writer);
+				// note: Ignore return value for now since errors should already be reported
+				/*bool bShaderMapValid = */ KernelResource->SerializeShaderMap(Writer);
 
 				FFileHelper::SaveArrayToFile(ResultData, *ShaderMapOutputFilepath);
 
@@ -375,7 +376,14 @@ namespace UE::IREE::Compiler
 					if (FeatureLevelType != ERHIFeatureLevel::Type::SM6)
 					{
 						const FString FeatureLevelName = LexToString(FeatureLevelType);
-						UE_LOG(LogIREEDriverRDG, Warning, TEXT("Minimum RHI Feature level is SM6 (desired feature level %s)!"), *FeatureLevelName);
+						UE_LOG(LogIREEDriverRDG, Warning, TEXT("Shader platform %s: Minimum RHI Feature level is SM6 (desired feature level %s)!"), *ShaderPlatformName, *FeatureLevelName);
+
+						continue;
+					}
+
+					if (FDataDrivenShaderPlatformInfo::GetSupportsRealTypes(ShaderPlatform) == ERHIFeatureSupport::Unsupported)
+					{
+						UE_LOG(LogIREEDriverRDG, Warning, TEXT("Shader platform %s does not support 16-bit types!"), *ShaderPlatformName);
 
 						continue;
 					}
diff --git a/Engine/Plugins/Experimental/NNERuntimeIREE/Source/IREEDriverRDG/Private/IREEDriverRDGDirectCommandBuffer.cpp b/Engine/Plugins/Experimental/NNERuntimeIREE/Source/IREEDriverRDG/Private/IREEDriverRDGDirectCommandBuffer.cpp
--- a/Engine/Plugins/Experimental/NNERuntimeIREE/Source/IREEDriverRDG/Private/IREEDriverRDGDirectCommandBuffer.cpp
+++ b/Engine/Plugins/Experimental/NNERuntimeIREE/Source/IREEDriverRDG/Private/IREEDriverRDGDirectCommandBuffer.cpp
@@ -211,12 +211,10 @@ private:
 
 		if (SourceRDGBuffer == TargetRDGBuffer)
 		{
-			const uint32 NumBytes = SourceRef.length;
-			FRDGBufferDesc TmpBufferDesc = FRDGBufferDesc::CreateByteAddressDesc(NumBytes);
-			FRDGBufferRef TmpBuffer = GraphBuilder.CreateBuffer(TmpBufferDesc, TEXT("IREE::CopyBufferTmp"));
+			FRDGBufferRef TmpBuffer = GraphBuilder.CreateBuffer(SourceRDGBuffer->Desc, TEXT("IREE::CopyBufferTmp"));
 
-			AddCopyBufferPass(GraphBuilder, TmpBuffer, 0, SourceRDGBuffer, SourceRef.offset, NumBytes);
-			AddCopyBufferPass(GraphBuilder, TargetRDGBuffer, TargetRef.offset, TmpBuffer, 0, NumBytes);
+			AddCopyBufferPass(GraphBuilder, TmpBuffer, 0, SourceRDGBuffer, SourceRef.offset, SourceRef.length);
+			AddCopyBufferPass(GraphBuilder, TargetRDGBuffer, TargetRef.offset, TmpBuffer, 0, SourceRef.length);
 		}
 		else
 		{
@@ -336,18 +334,24 @@ private:
 		}
 
 		TShaderRef<FNNERuntimeIREEShader> Shader = KernelResource->GetShader(0);
-
-		RDG_EVENT_SCOPE_STAT(GraphBuilder, FDirectCommandBufferDispatch, "DirectCommandBuffer.Dispatch %s", *KernelName);
-		RDG_GPU_STAT_SCOPE(GraphBuilder, FDirectCommandBufferDispatch);
-
-		FComputeShaderUtils::AddPass(
-			GraphBuilder,
-			RDG_EVENT_NAME("DirectCommandBuffer.Dispatch"),
-			ERDGPassFlags::Compute | ERDGPassFlags::NeverCull,
-			Shader,
-			ShaderParameterMetadata,
-			ShaderParameters,
-			FIntVector(WorkgroupCount[0], WorkgroupCount[1], WorkgroupCount[2]));
+		if (Shader.IsValid())
+		{
+			RDG_EVENT_SCOPE_STAT(GraphBuilder, FDirectCommandBufferDispatch, "DirectCommandBuffer.Dispatch %s", *KernelName);
+			RDG_GPU_STAT_SCOPE(GraphBuilder, FDirectCommandBufferDispatch);
+
+			FComputeShaderUtils::AddPass(
+				GraphBuilder,
+				RDG_EVENT_NAME("DirectCommandBuffer.Dispatch"),
+				ERDGPassFlags::Compute | ERDGPassFlags::NeverCull,
+				Shader,
+				ShaderParameterMetadata,
+				ShaderParameters,
+				FIntVector(WorkgroupCount[0], WorkgroupCount[1], WorkgroupCount[2]));
+		}
+		else
+		{
+			UE_LOG(LogIREEDriverRDG, Warning, TEXT("%s: Missing shader for executable %s."), StringCast<TCHAR>(__FUNCTION__).Get(), *KernelName);
+		}
 
 		for (auto BufferCopyResources : BufferCopiesToAppend)
 		{
diff --git a/Engine/Plugins/Experimental/NNERuntimeIREE/Source/IREEDriverRDG/Private/IREEDriverRDGExecutable.cpp b/Engine/Plugins/Experimental/NNERuntimeIREE/Source/IREEDriverRDG/Private/IREEDriverRDGExecutable.cpp
--- a/Engine/Plugins/Experimental/NNERuntimeIREE/Source/IREEDriverRDG/Private/IREEDriverRDGExecutable.cpp
+++ b/Engine/Plugins/Experimental/NNERuntimeIREE/Source/IREEDriverRDG/Private/IREEDriverRDGExecutable.cpp
@@ -35,32 +35,16 @@ namespace UE::IREE::HAL::RDG
 namespace Private
 {
 
-TArray<uint32> GetBufferBindings(const FShaderParametersMetadata* ShaderParametersMetadata)
+TArray<uint32> GetBufferBindings(const FIREEDriverRDGShaderParametersMetadata& Metadata)
 {
 	TArray<uint32> Result;
-
-	const TArray<FShaderParametersMetadata::FMember> &Members = ShaderParametersMetadata->GetMembers();
-	for (int32 MemberIdx = 0; MemberIdx < Members.Num(); MemberIdx++)
+	for (const FIREEDriverRDGShaderParametersMetadataEntry& Entry : Metadata.Entries)
 	{
-		const FShaderParametersMetadata::FMember& Member = Members[MemberIdx];
-
-		const FString MemberName(Member.GetName());
-
-		EUniformBufferBaseType BaseType = Member.GetBaseType();
-		if (Member.IsVariableNativeType())
-		{
-			check(MemberName == TEXT("Constant"));
-		}
-		else
+		if (Entry.Type == FIREEDriverRDGUniformBufferBaseType::BUFFER_UAV)
 		{
-			check(MemberName.StartsWith(TEXT("Buffer")));
-			check(BaseType == EUniformBufferBaseType::UBMT_RDG_BUFFER_UAV);
-
-			const int32 BindingIndex = FCString::Atoi(*MemberName.Mid(6, 1));
-			Result.Add(BindingIndex);
+			Result.Add(Entry.Binding);
 		}
 	}
-
 	return Result;
 }
 
@@ -125,9 +109,14 @@ public:
 				MoveTemp(ShaderParameterMetadataAllocations),
 				ShaderParametersMetadata,
 				FName(),
-				GetBufferBindings(ShaderParametersMetadata)
+				GetBufferBindings(IREEShaderParametersMetadata)
 			);
-			KernelResource->SerializeShaderMap(Reader);
+			
+
+			if (!KernelResource->SerializeShaderMap(Reader))
+			{
+				return iree_make_status(IREE_STATUS_NOT_FOUND, "Loaded shader map is not valid/complete.");
+			}
 
 			Executable->KernelResources.Add(MoveTemp(KernelResource));
 		}
diff --git a/Engine/Plugins/Experimental/NNERuntimeIREE/Source/IREEDriverRDG/Private/IREEDriverRDGShaderParametersMetadata.cpp b/Engine/Plugins/Experimental/NNERuntimeIREE/Source/IREEDriverRDG/Private/IREEDriverRDGShaderParametersMetadata.cpp
--- a/Engine/Plugins/Experimental/NNERuntimeIREE/Source/IREEDriverRDG/Private/IREEDriverRDGShaderParametersMetadata.cpp
+++ b/Engine/Plugins/Experimental/NNERuntimeIREE/Source/IREEDriverRDG/Private/IREEDriverRDGShaderParametersMetadata.cpp
@@ -32,6 +32,8 @@ FIREEDriverRDGUniformBufferElementType GetElementType(const FString& TypeStr)
 
 bool BuildIREEShaderParametersMetadata(const FString& Filepath, FIREEDriverRDGShaderParametersMetadata& Metadata)
 {
+	constexpr int32 NUM_COLUMNS = 7;
+
 	FString Filedata;
 	if (!FFileHelper::LoadFileToString(Filedata, *Filepath))
 	{
@@ -60,9 +62,9 @@ bool BuildIREEShaderParametersMetadata(const FString& Filepath, FIREEDriverRDGSh
 		TArray<FString> Row;
 		Lines[i].ParseIntoArray(Row, TEXT(";"), false);
 
-		if (Row.Num() != 5)
+		if (Row.Num() != NUM_COLUMNS)
 		{
-			UE_LOG(LogIREEDriverRDG, Log, TEXT("Expected 5 elements"));
+			UE_LOG(LogIREEDriverRDG, Log, TEXT("Expected %d elements"), NUM_COLUMNS);
 			return false;
 		}
 
@@ -70,8 +72,10 @@ bool BuildIREEShaderParametersMetadata(const FString& Filepath, FIREEDriverRDGSh
 		Entry.Type = GetBaseType(Row[0]);
 		Entry.Name = Row[1];
 		Entry.ShaderType = Row[2];
-		Entry.ElementType = Row[3].IsEmpty() ? FIREEDriverRDGUniformBufferElementType::NONE : GetElementType(Row[3]);
-		Entry.NumElements = Row[4].IsEmpty() ? 0 : FCString::Atoi(*Row[4]);
+		Entry.Binding = FCString::Atoi(*Row[3]);
+		Entry.DescriptorSet = FCString::Atoi(*Row[4]);
+		Entry.ElementType = Row[5].IsEmpty() ? FIREEDriverRDGUniformBufferElementType::NONE : GetElementType(Row[5]);
+		Entry.NumElements = Row[6].IsEmpty() ? 0 : FCString::Atoi(*Row[6]);
 		
 		Metadata.Entries.Add(MoveTemp(Entry));
 	}
diff --git a/Engine/Plugins/Experimental/NNERuntimeIREE/Source/IREEDriverRDG/Private/IREEDriverRDGShaderParametersMetadata.h b/Engine/Plugins/Experimental/NNERuntimeIREE/Source/IREEDriverRDG/Private/IREEDriverRDGShaderParametersMetadata.h
--- a/Engine/Plugins/Experimental/NNERuntimeIREE/Source/IREEDriverRDG/Private/IREEDriverRDGShaderParametersMetadata.h
+++ b/Engine/Plugins/Experimental/NNERuntimeIREE/Source/IREEDriverRDG/Private/IREEDriverRDGShaderParametersMetadata.h
@@ -34,6 +34,12 @@ struct FIREEDriverRDGShaderParametersMetadataEntry
 	UPROPERTY()
 	FString ShaderType;
 
+	UPROPERTY()
+	uint32 Binding = 0;
+	
+	UPROPERTY()
+	uint32 DescriptorSet = 0;
+
 	UPROPERTY()
 	FIREEDriverRDGUniformBufferElementType ElementType = FIREEDriverRDGUniformBufferElementType::NONE;
 
diff --git a/Engine/Plugins/Experimental/NNERuntimeIREE/Source/NNERuntimeIREE/Private/NNERuntimeIREE.cpp b/Engine/Plugins/Experimental/NNERuntimeIREE/Source/NNERuntimeIREE/Private/NNERuntimeIREE.cpp
--- a/Engine/Plugins/Experimental/NNERuntimeIREE/Source/NNERuntimeIREE/Private/NNERuntimeIREE.cpp
+++ b/Engine/Plugins/Experimental/NNERuntimeIREE/Source/NNERuntimeIREE/Private/NNERuntimeIREE.cpp
@@ -25,6 +25,7 @@
 #include "NNERuntimeIREEMetaData.h"
 #include "NNERuntimeIREEModel.h"
 #include "NNERuntimeIREEModelData.h"
+#include "RHIGlobals.h"
 #include "Serialization/MemoryReader.h"
 #include "Serialization/MemoryWriter.h"
 
@@ -671,14 +672,14 @@ TSharedPtr<UE::NNE::FSharedModelData> UNNERuntimeIREERdg::CreateModelData(const
 			const FString VmfbFilePathSrc = FPaths::Combine(SharedLibrarySubDirFullPath, Info.VmfbFileName);
 			const FString VmfbFilePathDest = FPaths::Combine(StagingSubDirFullPath, Info.VmfbFileName);
 
-			IFileManager::Get().Copy(*VmfbFilePathDest, *VmfbFilePathSrc, bNeedCompileMlir);
+			IFileManager::Get().Copy(*VmfbFilePathDest, *VmfbFilePathSrc);
 
 			for (const FString& ShaderMapFilename : Info.SharedLibraryFileNames)
 			{
 				const FString SharedLibraryFilePathSrc = FPaths::Combine(SharedLibrarySubDirFullPath, ShaderMapFilename);
 				const FString SharedLibraryFilePathDest = FPaths::Combine(StagingSubDirFullPath, ShaderMapFilename);
 
-				IFileManager::Get().Copy(*SharedLibraryFilePathDest, *SharedLibraryFilePathSrc, bNeedCompileMlir);
+				IFileManager::Get().Copy(*SharedLibraryFilePathDest, *SharedLibraryFilePathSrc);
 			}
 		}
 	}
@@ -807,6 +808,20 @@ TSharedPtr<UE::NNE::IModelRDG> UNNERuntimeIREERdg::CreateModelRDG(const TObjectP
 
 bool UNNERuntimeIREERdg::IsAvailable() const
 {
+#if !WITH_EDITOR
+	if(GMaxRHIFeatureLevel != ERHIFeatureLevel::SM6)
+	{
+		UE_LOG(LogNNERuntimeIREE, Log, TEXT("Minimum feature level required is SM6 for current RHI platform."));
+		return false;
+	}
+
+	if(!GRHIGlobals.SupportsNative16BitOps)
+	{
+		UE_LOG(LogNNERuntimeIREE, Log, TEXT("Current RHI platform doesn't support native 16-bit operations."));
+		return false;
+	}
+#endif
+
 	return WITH_NNE_RUNTIME_IREE_RDG;
 }
 
diff --git a/Engine/Plugins/Experimental/NNERuntimeIREE/Source/NNERuntimeIREEShader/Internal/NNERuntimeIREEShaderShared.h b/Engine/Plugins/Experimental/NNERuntimeIREE/Source/NNERuntimeIREEShader/Internal/NNERuntimeIREEShaderShared.h
--- a/Engine/Plugins/Experimental/NNERuntimeIREE/Source/NNERuntimeIREEShader/Internal/NNERuntimeIREEShaderShared.h
+++ b/Engine/Plugins/Experimental/NNERuntimeIREE/Source/NNERuntimeIREEShader/Internal/NNERuntimeIREEShaderShared.h
@@ -266,7 +266,7 @@ public:
 	 */
 	virtual bool ShouldCache(EShaderPlatform InPlatform, const FShaderType* InShaderType) const;
 
-	void SerializeShaderMap(FArchive& Ar);
+	bool SerializeShaderMap(FArchive& Ar);
 
 	void GetDependentShaderTypes(EShaderPlatform InPlatform, TArray<FShaderType*>& OutShaderTypes) const;
 	virtual void GetShaderMapId(EShaderPlatform InPlatform, const ITargetPlatform* TargetPlatform, FNNERuntimeIREEShaderMapId& OutId) const;
diff --git a/Engine/Plugins/Experimental/NNERuntimeIREE/Source/NNERuntimeIREEShader/Internal/NNERuntimeIREEShaderType.h b/Engine/Plugins/Experimental/NNERuntimeIREE/Source/NNERuntimeIREEShader/Internal/NNERuntimeIREEShaderType.h
--- a/Engine/Plugins/Experimental/NNERuntimeIREE/Source/NNERuntimeIREEShader/Internal/NNERuntimeIREEShaderType.h
+++ b/Engine/Plugins/Experimental/NNERuntimeIREE/Source/NNERuntimeIREEShader/Internal/NNERuntimeIREEShaderType.h
@@ -153,7 +153,7 @@ protected:
 	{
 		ModifyCompilationEnvironment(FNNERuntimeIREEShaderPermutationParameters(InPlatform), OutEnvironment);
 
-		if (FDataDrivenShaderPlatformInfo::GetSupportsRealTypes(InPlatform) == ERHIFeatureSupport::RuntimeGuaranteed)
+		if (FDataDrivenShaderPlatformInfo::GetSupportsRealTypes(InPlatform) != ERHIFeatureSupport::Unsupported)
 		{
 			OutEnvironment.CompilerFlags.Add(CFLAG_AllowRealTypes);
 		}
diff --git a/Engine/Plugins/Experimental/NNERuntimeIREE/Source/NNERuntimeIREEShader/Private/NNERuntimeIREEShaderShared.cpp b/Engine/Plugins/Experimental/NNERuntimeIREE/Source/NNERuntimeIREEShader/Private/NNERuntimeIREEShaderShared.cpp
--- a/Engine/Plugins/Experimental/NNERuntimeIREE/Source/NNERuntimeIREEShader/Private/NNERuntimeIREEShaderShared.cpp
+++ b/Engine/Plugins/Experimental/NNERuntimeIREE/Source/NNERuntimeIREEShader/Private/NNERuntimeIREEShaderShared.cpp
@@ -37,14 +37,16 @@ bool FNNERuntimeIREEResource::ShouldCache(EShaderPlatform InPlatform, const FSha
 	return true;
 }
 
-void FNNERuntimeIREEResource::SerializeShaderMap(FArchive& Ar)
+bool FNNERuntimeIREEResource::SerializeShaderMap(FArchive& Ar)
 {
+	bool bSuccess = false;
+
 	if (Ar.IsSaving())
 	{
 #if WITH_EDITOR
 		FinishCompilation();
 
-		bool bValid = GameThreadShaderMap != nullptr && GameThreadShaderMap->CompiledSuccessfully();
+		bool bValid = bSuccess = GameThreadShaderMap != nullptr && GameThreadShaderMap->CompiledSuccessfully();
 		Ar << bValid;
 
 		if (bValid)
@@ -62,17 +64,19 @@ void FNNERuntimeIREEResource::SerializeShaderMap(FArchive& Ar)
 		if (bValid)
 		{
 			TRefCountPtr<FNNERuntimeIREEShaderMap> LoadedShaderMap = new FNNERuntimeIREEShaderMap();
-			bool bSuccessfullyLoaded = LoadedShaderMap->Serialize(Ar);
+			bSuccess = LoadedShaderMap->Serialize(Ar);
 
 			// Toss the loaded shader data if this is a server only instance or if it's for a different RHI than the current one.
 			// todo[CF] Don't cook it in the first place
-			if (bSuccessfullyLoaded && FApp::CanEverRender())
+			if (bSuccess && FApp::CanEverRender())
 			{
 				GameThreadShaderMap = RenderingThreadShaderMap = LoadedShaderMap;
 				GameThreadShaderMap->GetResource()->SetOwnerName(GetOwnerFName());
 			}
 		}
 	}
+
+	return bSuccess;
 }
 
 void FNNERuntimeIREEResource::SetRenderingThreadShaderMap(FNNERuntimeIREEShaderMap* InShaderMap)
