[NNE] NNERuntimeIREE updated meta data parser to accept MLIR function operation util.func.

#tests Editor, Win64, model tests

#rnx

diff --git a/Engine/Plugins/Experimental/NNERuntimeIREE/Source/NNERuntimeIREE/Private/NNERuntimeIREEMetaData.cpp b/Engine/Plugins/Experimental/NNERuntimeIREE/Source/NNERuntimeIREE/Private/NNERuntimeIREEMetaData.cpp
--- a/Engine/Plugins/Experimental/NNERuntimeIREE/Source/NNERuntimeIREE/Private/NNERuntimeIREEMetaData.cpp
+++ b/Engine/Plugins/Experimental/NNERuntimeIREE/Source/NNERuntimeIREE/Private/NNERuntimeIREEMetaData.cpp
@@ -359,7 +359,7 @@ bool UNNERuntimeIREEModuleMetaData::ParseFromString(const FString& ModuleString)
 	// @: Have exactly one @
 	// [^(]*: Skip zero or more characters not being (
 	// (: Have exactly one (
-	const FRegexPattern FunctionNamePattern(TEXT("func\\.func[^@]*@[^(]*\\("));
+	const FRegexPattern FunctionNamePattern(TEXT("[func|util]\\.func[^@]*@[^(]*\\("));
 	FRegexMatcher FunctionNameMatcher(FunctionNamePattern, *ModuleString);
 	while (FunctionNameMatcher.FindNext())
 	{
