[NNE] NNERuntimeIREERdg fixed overflow bug.

#tests Editor, Win64 on Vulkan, verified with debugger.

#rnx

diff --git a/Engine/Plugins/Experimental/NNERuntimeIREE/Source/IREEDriverRDG/Private/IREEDriverRDGBuiltinExecutables.cpp b/Engine/Plugins/Experimental/NNERuntimeIREE/Source/IREEDriverRDG/Private/IREEDriverRDGBuiltinExecutables.cpp
--- a/Engine/Plugins/Experimental/NNERuntimeIREE/Source/IREEDriverRDG/Private/IREEDriverRDGBuiltinExecutables.cpp
+++ b/Engine/Plugins/Experimental/NNERuntimeIREE/Source/IREEDriverRDG/Private/IREEDriverRDGBuiltinExecutables.cpp
@@ -81,13 +81,13 @@ iree_status_t AddFillBufferPass(FRDGBuilder& GraphBuilder, FRDGBufferRef RDGBuff
 	const uint32 ExpandedPattern = Private::ExpandPatternToWord(Pattern, PatternLength);
 	const uint32 ExpandedLength = 4;
 
-	const int32 NumThreads = (int32)FillLength / 4;
-	const int32 NumThreadsPerDispatch = GRHIMaxDispatchThreadGroupsPerDimension.X * FFillBufferConstants::THREAD_GROUP_SIZE;
-	const int32 NumDispatches = FMath::DivideAndRoundUp(NumThreads, NumThreadsPerDispatch);
+	const uint32 NumThreads = FillLength / 4;
+	const uint64 NumThreadsPerDispatch = GRHIMaxDispatchThreadGroupsPerDimension.X * FFillBufferConstants::THREAD_GROUP_SIZE;
+	const uint32 NumDispatches = (uint32)FMath::DivideAndRoundUp((uint64)NumThreads, NumThreadsPerDispatch);
 
-	for (int32 i = 0; i < NumDispatches; i++)
+	for (uint32 i = 0; i < NumDispatches; i++)
 	{
-		const int32 DispatchNumThreads = FMath::Min((int32)NumThreads - i * NumThreadsPerDispatch, NumThreadsPerDispatch);
+		const uint32 DispatchNumThreads = (uint32)FMath::Min((uint64)NumThreads - i * NumThreadsPerDispatch, NumThreadsPerDispatch);
 		check(DispatchNumThreads > 0);
 
 		FFillBufferCS::FParameters* ShaderParameters = GraphBuilder.AllocParameters<FFillBufferCS::FParameters>();
diff --git a/Engine/Plugins/Experimental/NNERuntimeIREE/Source/NNERuntimeIREEShader/Internal/NNERuntimeIREEShaderFillBufferCS.h b/Engine/Plugins/Experimental/NNERuntimeIREE/Source/NNERuntimeIREEShader/Internal/NNERuntimeIREEShaderFillBufferCS.h
--- a/Engine/Plugins/Experimental/NNERuntimeIREE/Source/NNERuntimeIREEShader/Internal/NNERuntimeIREEShaderFillBufferCS.h
+++ b/Engine/Plugins/Experimental/NNERuntimeIREE/Source/NNERuntimeIREEShader/Internal/NNERuntimeIREEShaderFillBufferCS.h
@@ -12,7 +12,7 @@ namespace UE::NNERuntimeIREEShader::Internal
 class FFillBufferConstants
 {
 public:
-	static const int32 THREAD_GROUP_SIZE{ 256 };
+	static const uint32 THREAD_GROUP_SIZE{ 256 };
 };
 
 class NNERUNTIMEIREESHADER_API FFillBufferCS : public FGlobalShader
