[NNE] NNERuntimeIREERdg conclude VulkanRHI support.
- Updated IREE compiler binaries (added switch to use structured buffers)
- Added config files (similar to CPU runtime), containing tool commands, arguments and shader platform targets
- Adjusted folder structure for compiler output files
- Added ONNX import feature

#tests All platforms, Editor, Game ran all model tests

#rnx

diff --git a/Engine/Plugins/Experimental/NNERuntimeIREE/Config/IREERDG_Linux_To_Linux.json b/Engine/Plugins/Experimental/NNERuntimeIREE/Config/IREERDG_Linux_To_Linux.json
new file mode 100644
--- /dev/null
+++ b/Engine/Plugins/Experimental/NNERuntimeIREE/Config/IREERDG_Linux_To_Linux.json
@@ -0,0 +1,11 @@
+{
+	"ImporterCommand": [ "${PLUGIN_DIR}/Binaries/ThirdParty/IREE/Linux/torch-mlir-import-onnx" ],
+	"ImporterArguments": "-o ${OUTPUT_PATH} --opset-version 21 ${INPUT_PATH}",
+	"CompilerCommand": [ "${PLUGIN_DIR}/Binaries/ThirdParty/IREE/Linux/iree-compile" ],
+	"Targets": [
+		{
+			"ShaderPlatform": "VULKAN_SM6",
+			"CompilerArguments": "--iree-hal-target-backends=unreal-shader --iree-unreal-target=sm_75 --iree-unreal-use-structured-buffers --iree-hal-dump-executable-binaries-to=${BINARIES_PATH} -o ${VMFB_PATH} ${INPUT_PATH}"
+		}
+	]
+}
\ No newline at end of file
diff --git a/Engine/Plugins/Experimental/NNERuntimeIREE/Config/IREERDG_Mac_To_Mac.json b/Engine/Plugins/Experimental/NNERuntimeIREE/Config/IREERDG_Mac_To_Mac.json
new file mode 100644
--- /dev/null
+++ b/Engine/Plugins/Experimental/NNERuntimeIREE/Config/IREERDG_Mac_To_Mac.json
@@ -0,0 +1,11 @@
+{
+	"ImporterCommand": [ "${PLUGIN_DIR}/Binaries/ThirdParty/IREE/Mac/torch-mlir-import-onnx" ],
+	"ImporterArguments": "-o ${OUTPUT_PATH} --opset-version 21 ${INPUT_PATH}",
+	"CompilerCommand": [ "${PLUGIN_DIR}/Binaries/ThirdParty/IREE/Mac/iree-compile" ],
+	"Targets": [
+		{
+			"ShaderPlatform": "METAL_SM6",
+			"CompilerArguments": "--iree-hal-target-backends=unreal-shader --iree-unreal-target=sm_75 --iree-hal-dump-executable-binaries-to=${BINARIES_PATH} -o ${VMFB_PATH} ${INPUT_PATH}"
+		}
+	]
+}
\ No newline at end of file
diff --git a/Engine/Plugins/Experimental/NNERuntimeIREE/Config/IREERDG_Windows_To_Windows.json b/Engine/Plugins/Experimental/NNERuntimeIREE/Config/IREERDG_Windows_To_Windows.json
new file mode 100644
--- /dev/null
+++ b/Engine/Plugins/Experimental/NNERuntimeIREE/Config/IREERDG_Windows_To_Windows.json
@@ -0,0 +1,15 @@
+{
+	"ImporterCommand": [ "${PLUGIN_DIR}/Binaries/ThirdParty/IREE/Windows/torch-mlir-import-onnx.exe" ],
+	"ImporterArguments": "-o ${OUTPUT_PATH} --opset-version 21 ${INPUT_PATH}",
+	"CompilerCommand": [ "${PLUGIN_DIR}/Binaries/ThirdParty/IREE/Windows/iree-compile.exe" ],
+	"Targets": [
+		{
+			"ShaderPlatform": "VULKAN_SM6",
+			"CompilerArguments": "--iree-hal-target-backends=unreal-shader --iree-unreal-target=sm_75 --iree-unreal-use-structured-buffers --iree-hal-dump-executable-binaries-to=${BINARIES_PATH} -o ${VMFB_PATH} ${INPUT_PATH}"
+		},
+		{
+			"ShaderPlatform": "PCD3D_SM6",
+			"CompilerArguments": "--iree-hal-target-backends=unreal-shader --iree-unreal-target=sm_75 --iree-hal-dump-executable-binaries-to=${BINARIES_PATH} -o ${VMFB_PATH} ${INPUT_PATH}"
+		}
+	]
+}
\ No newline at end of file
diff --git a/Engine/Plugins/Experimental/NNERuntimeIREE/Config/Linux/LinuxGame.ini b/Engine/Plugins/Experimental/NNERuntimeIREE/Config/Linux/LinuxGame.ini
--- a/Engine/Plugins/Experimental/NNERuntimeIREE/Config/Linux/LinuxGame.ini
+++ b/Engine/Plugins/Experimental/NNERuntimeIREE/Config/Linux/LinuxGame.ini
@@ -1,2 +1,3 @@
 [/Script/UnrealEd.ProjectPackagingSettings]
-+DirectoriesToAlwaysStageAsNonUFS=(Path="../Binaries/Linux/NNERuntimeIREE")
++DirectoriesToAlwaysStageAsNonUFS=(Path="../Binaries/Linux/NNERuntimeIREECpu")
++DirectoriesToAlwaysStageAsNonUFS=(Path="../Binaries/Linux/NNERuntimeIREERdg")
diff --git a/Engine/Plugins/Experimental/NNERuntimeIREE/Config/LinuxArm64/LinuxArm64Game.ini b/Engine/Plugins/Experimental/NNERuntimeIREE/Config/LinuxArm64/LinuxArm64Game.ini
--- a/Engine/Plugins/Experimental/NNERuntimeIREE/Config/LinuxArm64/LinuxArm64Game.ini
+++ b/Engine/Plugins/Experimental/NNERuntimeIREE/Config/LinuxArm64/LinuxArm64Game.ini
@@ -1,2 +1,3 @@
 [/Script/UnrealEd.ProjectPackagingSettings]
-+DirectoriesToAlwaysStageAsNonUFS=(Path="../Binaries/LinuxArm64/NNERuntimeIREE")
++DirectoriesToAlwaysStageAsNonUFS=(Path="../Binaries/LinuxArm64/NNERuntimeIREECpu")
++DirectoriesToAlwaysStageAsNonUFS=(Path="../Binaries/LinuxArm64/NNERuntimeIREERdg")
diff --git a/Engine/Plugins/Experimental/NNERuntimeIREE/Config/Mac/MacGame.ini b/Engine/Plugins/Experimental/NNERuntimeIREE/Config/Mac/MacGame.ini
--- a/Engine/Plugins/Experimental/NNERuntimeIREE/Config/Mac/MacGame.ini
+++ b/Engine/Plugins/Experimental/NNERuntimeIREE/Config/Mac/MacGame.ini
@@ -1,2 +1,3 @@
 [/Script/UnrealEd.ProjectPackagingSettings]
-+DirectoriesToAlwaysStageAsNonUFS=(Path="../Binaries/Mac/NNERuntimeIREE")
++DirectoriesToAlwaysStageAsNonUFS=(Path="../Binaries/Mac/NNERuntimeIREECpu")
++DirectoriesToAlwaysStageAsNonUFS=(Path="../Binaries/Mac/NNERuntimeIREERdg")
diff --git a/Engine/Plugins/Experimental/NNERuntimeIREE/Config/Windows/WindowsGame.ini b/Engine/Plugins/Experimental/NNERuntimeIREE/Config/Windows/WindowsGame.ini
--- a/Engine/Plugins/Experimental/NNERuntimeIREE/Config/Windows/WindowsGame.ini
+++ b/Engine/Plugins/Experimental/NNERuntimeIREE/Config/Windows/WindowsGame.ini
@@ -1,2 +1,3 @@
 [/Script/UnrealEd.ProjectPackagingSettings]
-+DirectoriesToAlwaysStageAsNonUFS=(Path="../Binaries/Win64/NNERuntimeIREE")
++DirectoriesToAlwaysStageAsNonUFS=(Path="../Binaries/Win64/NNERuntimeIREECpu")
++DirectoriesToAlwaysStageAsNonUFS=(Path="../Binaries/Win64/NNERuntimeIREERdg")
diff --git a/Engine/Plugins/Experimental/NNERuntimeIREE/Source/IREEDriverRDG/Internal/IREECompilerRDG.h b/Engine/Plugins/Experimental/NNERuntimeIREE/Source/IREEDriverRDG/Internal/IREECompilerRDG.h
--- a/Engine/Plugins/Experimental/NNERuntimeIREE/Source/IREEDriverRDG/Internal/IREECompilerRDG.h
+++ b/Engine/Plugins/Experimental/NNERuntimeIREE/Source/IREEDriverRDG/Internal/IREECompilerRDG.h
@@ -3,6 +3,7 @@
 #pragma once
 
 #include "CoreMinimal.h"
+#include "RHIShaderPlatform.h"
 #include "Serialization/JsonSerializerMacros.h"
 
 #include "IREECompilerRDG.generated.h"
@@ -13,7 +14,7 @@ struct FIREECompilerRDGBuildTargetResult
 	GENERATED_BODY()
 
 	UPROPERTY()
-	FString Architecture;
+	FString ShaderPlatform;
 
 	UPROPERTY()
 	FString RelativeDirPath;
@@ -31,7 +32,7 @@ struct FIREECompilerRDGResult
 	GENERATED_BODY()
 
 	UPROPERTY()
-	TArray<FIREECompilerRDGBuildTargetResult> ArchitectureInfos;
+	TArray<FIREECompilerRDGBuildTargetResult> BuildTargetResult;
 };
 
 #ifdef WITH_IREE_DRIVER_RDG
@@ -41,28 +42,27 @@ namespace UE::IREE::Compiler::RDG
 {
 	struct FBuildTarget : public FJsonSerializable
 	{
-		FString Architecture;
+		FString ShaderPlatform;
 		FString CompilerArguments;
 		FString LinkerArguments;
 
 		BEGIN_JSON_SERIALIZER
-			JSON_SERIALIZE("Architecture", Architecture);
+			JSON_SERIALIZE("ShaderPlatform", ShaderPlatform);
 			JSON_SERIALIZE("CompilerArguments", CompilerArguments);
-			JSON_SERIALIZE("LinkerArguments", LinkerArguments);
 		END_JSON_SERIALIZER
 	};
 
 	struct FBuildConfig : public FJsonSerializable
 	{
+		TArray<FString> ImporterCommand;
+		FString ImporterArguments;
 		TArray<FString> CompilerCommand;
-		TArray<FString> LinkerCommand;
-		FString SharedLibExt;
 		TArray<FBuildTarget> BuildTargets;
 
 		BEGIN_JSON_SERIALIZER
+			JSON_SERIALIZE_ARRAY("ImporterCommand", ImporterCommand);
+			JSON_SERIALIZE("ImporterArguments", ImporterArguments);
 			JSON_SERIALIZE_ARRAY("CompilerCommand", CompilerCommand);
-			JSON_SERIALIZE_ARRAY("LinkerCommand", LinkerCommand);
-			JSON_SERIALIZE("SharedLibExt", SharedLibExt);
 			JSON_SERIALIZE_ARRAY_SERIALIZABLE("Targets", BuildTargets, FBuildTarget);
 		END_JSON_SERIALIZER
 	};
@@ -70,20 +70,21 @@ namespace UE::IREE::Compiler::RDG
 	class IREEDRIVERRDG_API FCompiler
 	{
 	private:
-		FCompiler(const ITargetPlatform* TargetPlatform, const FString& InCompilerCommand, const FString& InLinkerCommand, const FString& InSharedLibExt, TConstArrayView<FBuildTarget> InBuildTargets);
+		FCompiler(const ITargetPlatform* TargetPlatform, const FString& InImporterCommand, const FString& InImporterArguments, const FString& InCompilerCommand, const FString& InSharedLibExt, TConstArrayView<FBuildTarget> InBuildTargets);
 
 	public:
 		~FCompiler() {};
 		static TUniquePtr<FCompiler> Make(const ITargetPlatform* TargetPlatform);
 
 	public:
-		bool CompileMlir(TConstArrayView<uint8> InFileData, const FString& InModelName, const FString& InOutputDir, FIREECompilerRDGResult& OutCompilerResult);
+		bool ImportOnnx(TConstArrayView<uint8> InFileData, const FString& InModelName, const FString& InOutputDir, TArray64<uint8>& OutMlirData);
+		bool CompileMlir(TConstArrayView<uint8> InFileData, const FString& InModelName, const FString& InOutputDir, TConstArrayView<EShaderPlatform> ShaderPlatforms, FIREECompilerRDGResult& OutCompilerResult);
 
 	private:
 		const ITargetPlatform* TargetPlatform;
+		FString ImporterCommand;
+		FString ImporterArguments;
 		FString CompilerCommand;
-		FString LinkerCommand;
-		FString SharedLibExt;
 		TArray<FBuildTarget> BuildTargets;
 	};
 } // namespace UE::IREE::Compiler::RDG
diff --git a/Engine/Plugins/Experimental/NNERuntimeIREE/Source/IREEDriverRDG/Private/IREECompilerRDG.cpp b/Engine/Plugins/Experimental/NNERuntimeIREE/Source/IREEDriverRDG/Private/IREECompilerRDG.cpp
--- a/Engine/Plugins/Experimental/NNERuntimeIREE/Source/IREEDriverRDG/Private/IREECompilerRDG.cpp
+++ b/Engine/Plugins/Experimental/NNERuntimeIREE/Source/IREEDriverRDG/Private/IREECompilerRDG.cpp
@@ -35,6 +35,16 @@ namespace UE::IREE::Compiler
 	{
 		namespace Private
 		{
+			FString GetOnnxFilePath(const FString& InModelName, const FString& InOutputDir)
+			{
+				return FPaths::Combine(InOutputDir, InModelName) + ".onnx";
+			}
+
+			FString GetMlirFilePath(const FString& InModelName, const FString& InOutputDir)
+			{
+				return FPaths::Combine(InOutputDir, InModelName) + ".mlir";
+			}
+
 			bool ResolveEnvironmentVariables(FString& String)
 			{
 				const FString StartString = "$ENV{";
@@ -183,7 +193,7 @@ namespace UE::IREE::Compiler
 			}
 		} // Private
 
-		FCompiler::FCompiler(const ITargetPlatform* InTargetPlatform, const FString& InCompilerCommand, const FString& InLinkerCommand, const FString& InSharedLibExt, TConstArrayView<FBuildTarget> InBuildTargets) : TargetPlatform(InTargetPlatform), CompilerCommand(InCompilerCommand), LinkerCommand(InLinkerCommand), SharedLibExt(InSharedLibExt), BuildTargets(InBuildTargets)
+		FCompiler::FCompiler(const ITargetPlatform* InTargetPlatform, const FString& InImporterCommand, const FString& InImporterArguments, const FString& InCompilerCommand, const FString& InSharedLibExt, TConstArrayView<FBuildTarget> InBuildTargets) : TargetPlatform(InTargetPlatform), ImporterCommand(InImporterCommand), ImporterArguments(InImporterArguments), CompilerCommand(InCompilerCommand), BuildTargets(InBuildTargets)
 		{
 			
 		}
@@ -195,7 +205,7 @@ namespace UE::IREE::Compiler
 			const FString TargetPlatformName = TargetPlatform ? TargetPlatform->IniPlatformName() : UGameplayStatics::GetPlatformName();
 
 			FString PluginDir = FPaths::ConvertRelativePathToFull(*IPluginManager::Get().FindPlugin(UE_PLUGIN_NAME)->GetBaseDir());
-			FString BuildConfigFileName = FString("IREE_") + UGameplayStatics::GetPlatformName() + "_To_" + TargetPlatformName + ".json";
+			FString BuildConfigFileName = FString("IREERDG_") + UGameplayStatics::GetPlatformName() + "_To_" + TargetPlatformName + ".json";
 			TArray<FString> BuildConfigFilePaths =
 			{
 				FPaths::Combine(FPaths::ConvertRelativePathToFull(*FPaths::ProjectConfigDir()), BuildConfigFileName),
@@ -204,8 +214,9 @@ namespace UE::IREE::Compiler
 				FPaths::Combine(FPaths::ConvertRelativePathToFull(*FPaths::EngineDir()), "Platforms", TargetPlatformName, "Plugins", "Experimental", UE_PLUGIN_NAME, "Config", BuildConfigFileName)
 			};
 
+			FString ImporterCommand;
+			FString ImporterArguments;
 			FString CompilerCommand;
-			FString LinkerCommand;
 			FString SharedLibExt;
 			TArray<FBuildTarget> BuildTargets;
 			IPlatformFile& PlatformFile = FPlatformFileManager::Get().GetPlatformFile();
@@ -227,57 +238,57 @@ namespace UE::IREE::Compiler
 								continue;
 							}
 
-							FString TmpCompilerCommand;
-							for (int32 i = 0; i < BuildConfig.CompilerCommand.Num(); i++)
+							FString TmpImporterCommand;
+							for (int32 i = 0; i < BuildConfig.ImporterCommand.Num(); i++)
 							{
-								if (ResolveEnvironmentVariables(BuildConfig.CompilerCommand[i]))
+								if (ResolveEnvironmentVariables(BuildConfig.ImporterCommand[i]))
 								{
-									BuildConfig.CompilerCommand[i].ReplaceInline(*FString("${PLUGIN_DIR}"), *PluginDir);
-									BuildConfig.CompilerCommand[i].ReplaceInline(*FString("${PROJECT_DIR}"), *FPaths::ProjectDir());
-									if (PlatformFile.FileExists(*BuildConfig.CompilerCommand[i]))
+									BuildConfig.ImporterCommand[i].ReplaceInline(*FString("${PLUGIN_DIR}"), *PluginDir);
+									BuildConfig.ImporterCommand[i].ReplaceInline(*FString("${PROJECT_DIR}"), *FPaths::ProjectDir());
+									if (PlatformFile.FileExists(*BuildConfig.ImporterCommand[i]))
 									{
-										TmpCompilerCommand = BuildConfig.CompilerCommand[i];
+										TmpImporterCommand = BuildConfig.ImporterCommand[i];
 										break;
 									}
 								}
 								else
 								{
-									UE_LOG(LogIREEDriverRDG, Warning, TEXT("IREECompilerRDG could not replace environment variables in %s"), *BuildConfig.CompilerCommand[i]);
+									UE_LOG(LogIREEDriverRDG, Warning, TEXT("IREECompilerRDG could not replace environment variables in %s"), *BuildConfig.ImporterCommand[i]);
 								}
 							}
-							if (TmpCompilerCommand.IsEmpty())
+							if (TmpImporterCommand.IsEmpty())
 							{
-								UE_LOG(LogIREEDriverRDG, Warning, TEXT("IREECompilerRDG could not find the compiler executable in %s"), *BuildConfigFilePath);
+								UE_LOG(LogIREEDriverRDG, Warning, TEXT("IREECompilerRDG could not find the importer executable in %s"), *BuildConfigFilePath);
 								continue;
 							}
 
-							FString TmpLinkerCommand;
-							for (int32 i = 0; i < BuildConfig.LinkerCommand.Num(); i++)
+							FString TmpCompilerCommand;
+							for (int32 i = 0; i < BuildConfig.CompilerCommand.Num(); i++)
 							{
-								if (ResolveEnvironmentVariables(BuildConfig.LinkerCommand[i]))
+								if (ResolveEnvironmentVariables(BuildConfig.CompilerCommand[i]))
 								{
-									BuildConfig.LinkerCommand[i].ReplaceInline(*FString("${PLUGIN_DIR}"), *PluginDir);
-									BuildConfig.LinkerCommand[i].ReplaceInline(*FString("${PROJECT_DIR}"), *FPaths::ProjectDir());
-									if (PlatformFile.FileExists(*BuildConfig.LinkerCommand[i]))
+									BuildConfig.CompilerCommand[i].ReplaceInline(*FString("${PLUGIN_DIR}"), *PluginDir);
+									BuildConfig.CompilerCommand[i].ReplaceInline(*FString("${PROJECT_DIR}"), *FPaths::ProjectDir());
+									if (PlatformFile.FileExists(*BuildConfig.CompilerCommand[i]))
 									{
-										TmpLinkerCommand = BuildConfig.LinkerCommand[i];
+										TmpCompilerCommand = BuildConfig.CompilerCommand[i];
 										break;
 									}
 								}
 								else
 								{
-									UE_LOG(LogIREEDriverRDG, Warning, TEXT("IREECompilerRDG could not replace environment variables in %s"), *BuildConfig.LinkerCommand[i]);
+									UE_LOG(LogIREEDriverRDG, Warning, TEXT("IREECompilerRDG could not replace environment variables in %s"), *BuildConfig.CompilerCommand[i]);
 								}
 							}
-							if (TmpLinkerCommand.IsEmpty())
+							if (TmpCompilerCommand.IsEmpty())
 							{
-								UE_LOG(LogIREEDriverRDG, Warning, TEXT("IREECompilerRDG could not find the linker executable in %s"), *BuildConfigFilePath);
+								UE_LOG(LogIREEDriverRDG, Warning, TEXT("IREECompilerRDG could not find the compiler executable in %s"), *BuildConfigFilePath);
 								continue;
 							}
 
+							ImporterCommand = TmpImporterCommand;
+							ImporterArguments = BuildConfig.ImporterArguments;
 							CompilerCommand = TmpCompilerCommand;
-							LinkerCommand = TmpLinkerCommand;
-							SharedLibExt = BuildConfig.SharedLibExt;
 							BuildTargets = BuildConfig.BuildTargets;
 							break;
 						}
@@ -292,57 +303,111 @@ namespace UE::IREE::Compiler
 					}
 				}
 			}
-			if (CompilerCommand.IsEmpty() || LinkerCommand.IsEmpty() || BuildTargets.IsEmpty())
+			if (CompilerCommand.IsEmpty() || BuildTargets.IsEmpty())
 			{
 				return TUniquePtr<FCompiler>();
 			}
-			return TUniquePtr<FCompiler>(new FCompiler(TargetPlatform, CompilerCommand, LinkerCommand, SharedLibExt, BuildTargets));
+			return TUniquePtr<FCompiler>(new FCompiler(TargetPlatform, ImporterCommand, ImporterArguments, CompilerCommand, SharedLibExt, BuildTargets));
 		}
 
-		bool FCompiler::CompileMlir(TConstArrayView<uint8> InFileData, const FString& InModelName, const FString& InOutputDir, FIREECompilerRDGResult& OutCompilerResult)
+		bool FCompiler::ImportOnnx(TConstArrayView<uint8> InFileData, const FString& InModelName, const FString& InOutputDir, TArray64<uint8>& OutMlirData)
 		{
-			using namespace Private;
+			SCOPED_NAMED_EVENT_TEXT("FCompiler::ImportOnnx", FColor::Magenta);
 
-			UE_LOG(LogIREEDriverRDG, Log, TEXT("FCompiler::CompileMlir %s"), *InModelName);
+			using namespace Private;
 
 			IPlatformFile& PlatformFile = FPlatformFileManager::Get().GetPlatformFile();
 
-			FString InputFilePath = FPaths::Combine(InOutputDir, InModelName) + ".mlir";
+			FString InputFilePath = GetOnnxFilePath(InModelName, InOutputDir);
 			if (!PlatformFile.FileExists(*InputFilePath))
 			{
 				SCOPED_NAMED_EVENT_TEXT("InputFile", FColor::Magenta);
 
-				FFileHelper::SaveArrayToFile(InFileData, *InputFilePath);
+				if(!FFileHelper::SaveArrayToFile(InFileData, *InputFilePath))
+				{
+					UE_LOG(LogIREEDriverRDG, Warning, TEXT("IREECompilerRDG failed to save ONNX model \"%s\""), *InputFilePath);
+					return false;
+				}
+			}
+
+			FString OutputFilePath = GetMlirFilePath(InModelName, InOutputDir);
+			FString IntermediateFilePathNoExt = FPaths::Combine(InOutputDir, InModelName);
+
+			if (!ResolveEnvironmentVariables(ImporterArguments))
+			{
+				UE_LOG(LogIREEDriverRDG, Warning, TEXT("IREECompilerRDG could not replace environment variables in %s"), *ImporterArguments);
+				return false;
+			}
+			ImporterArguments.ReplaceInline(*FString("${INPUT_PATH}"), *(FString("\"") + InputFilePath + "\""));
+			ImporterArguments.ReplaceInline(*FString("${OUTPUT_PATH}"), *(FString("\"") + OutputFilePath + "\""));
+
+			{
+				SCOPED_NAMED_EVENT_TEXT("Import", FColor::Magenta);
+
+				RunCommand(ImporterCommand, ImporterArguments, IntermediateFilePathNoExt + "_import-log.txt");
 			}
 
-			TArray<EShaderPlatform> ShaderPlatforms;
-			if (TargetPlatform)
+			if (!PlatformFile.FileExists(*OutputFilePath))
 			{
-				TArray<FName> DesiredShaderFormats;
-				TargetPlatform->GetAllTargetedShaderFormats(DesiredShaderFormats);
+				UE_LOG(LogIREEDriverRDG, Warning, TEXT("IREECompilerRDG failed to import the model \"%s\" using the command:"), *InputFilePath);
+				UE_LOG(LogIREEDriverRDG, Warning, TEXT("\"%s\" %s"), *ImporterCommand, *ImporterArguments);
+				return false;
+			}
 
-				for (const FName& ShaderFormatName : DesiredShaderFormats)
+			{
+				SCOPED_NAMED_EVENT_TEXT("Load", FColor::Magenta);
+				if(!FFileHelper::LoadFileToArray(OutMlirData, *OutputFilePath))
 				{
-					ShaderPlatforms.Add(ShaderFormatToLegacyShaderPlatform(ShaderFormatName));
+					UE_LOG(LogIREEDriverRDG, Warning, TEXT("IREECompilerRDG failed to load imported model \"%s\""), *OutputFilePath);
+					return false;
 				}
 			}
-			else
+
+			return true;
+		}
+
+		bool FCompiler::CompileMlir(TConstArrayView<uint8> InFileData, const FString& InModelName, const FString& InOutputDir, TConstArrayView<EShaderPlatform> ShaderPlatforms, FIREECompilerRDGResult& OutCompilerResult)
+		{
+			using namespace Private;
+
+			UE_LOG(LogIREEDriverRDG, Log, TEXT("FCompiler::CompileMlir %s"), *InModelName);
+
+			IPlatformFile& PlatformFile = FPlatformFileManager::Get().GetPlatformFile();
+
+			FString InputFilePath = FPaths::Combine(InOutputDir, InModelName) + ".mlir";
+			if (!PlatformFile.FileExists(*InputFilePath))
 			{
-				const ERHIFeatureLevel::Type CacheFeatureLevel = GMaxRHIFeatureLevel;
-				const EShaderPlatform ShaderPlatform = GShaderPlatformForFeatureLevel[CacheFeatureLevel];
+				SCOPED_NAMED_EVENT_TEXT("InputFile", FColor::Magenta);
 
-				ShaderPlatforms.Add(ShaderPlatform);
+				FFileHelper::SaveArrayToFile(InFileData, *InputFilePath);
 			}
 
 			bool bResult = true;
-			for (int32 i = 0; i < BuildTargets.Num(); i++)
+			for (EShaderPlatform ShaderPlatform : ShaderPlatforms)
 			{
-				FString IntermediateDirPath = FPaths::Combine(InOutputDir, BuildTargets[i].Architecture);
+				const FString ShaderPlatformName = LexToString(ShaderPlatform);
+
+				const FBuildTarget *Target = BuildTargets.FindByPredicate([ShaderPlatformName] (const FBuildTarget &Element) { return Element.ShaderPlatform == ShaderPlatformName; });
+				if (!Target)
+				{
+					UE_LOG(LogIREEDriverRDG, Warning, TEXT("IREECompilerRDG could not find build target for shader platform %s"), *ShaderPlatformName);
+					// Don't fail, model data will not be available...
+					// bResult = false;
+					continue;
+				}
+
+				FString IntermediateDirPath = FPaths::Combine(InOutputDir, Target->ShaderPlatform);
 				PlatformFile.CreateDirectoryTree(*IntermediateDirPath);
 				FString IntermediateFilePathNoExt = FPaths::Combine(IntermediateDirPath, InModelName);
 				FString VmfbFilePath = IntermediateFilePathNoExt + ".vmfb";
 
-				FString CompilerArguments = TEXT("--iree-hal-target-backends=unreal-shader --iree-unreal-target=sm_75 --iree-hal-dump-executable-binaries-to=${BINARIES_PATH} -o ${VMFB_PATH} ${INPUT_PATH}");
+				FString CompilerArguments = Target->CompilerArguments;
+				if (!ResolveEnvironmentVariables(CompilerArguments))
+				{
+					UE_LOG(LogIREEDriverRDG, Warning, TEXT("IREECompilerRDG could not replace environment variables in %s"), *Target->CompilerArguments);
+					bResult = false;
+					continue;
+				}
 				CompilerArguments.ReplaceInline(*FString("${BINARIES_PATH}"), *(FString("\"") + IntermediateDirPath + "\""));
 				CompilerArguments.ReplaceInline(*FString("${VMFB_PATH}"), *(FString("\"") + VmfbFilePath + "\""));
 				CompilerArguments.ReplaceInline(*FString("${INPUT_PATH}"), *(FString("\"") + InputFilePath + "\""));
@@ -364,67 +429,64 @@ namespace UE::IREE::Compiler
 					UE_LOG(LogIREEDriverRDG, Warning, TEXT("No HLSL shader files generated!"));
 				}
 
-				TArray<FString> ShaderMapFilenames;
-				for (EShaderPlatform ShaderPlatform : ShaderPlatforms)
-				{
-					const FString ShaderPlatformName = LexToString(ShaderPlatform);
-					const FString OutDir = FPaths::Combine(IntermediateDirPath, ShaderPlatformName);
-
-					const FStaticFeatureLevel FeatureLevel = GetMaxSupportedFeatureLevel(ShaderPlatform);
-					ERHIFeatureLevel::Type FeatureLevelType = FeatureLevel.operator ERHIFeatureLevel::Type();
+				const FStaticFeatureLevel FeatureLevel = GetMaxSupportedFeatureLevel(ShaderPlatform);
+				ERHIFeatureLevel::Type FeatureLevelType = FeatureLevel.operator ERHIFeatureLevel::Type();
 
-					if (FeatureLevelType != ERHIFeatureLevel::Type::SM6)
-					{
-						const FString FeatureLevelName = LexToString(FeatureLevelType);
-						UE_LOG(LogIREEDriverRDG, Warning, TEXT("Shader platform %s: Minimum RHI Feature level is SM6 (desired feature level %s)!"), *ShaderPlatformName, *FeatureLevelName);
+				if (FeatureLevelType != ERHIFeatureLevel::Type::SM6)
+				{
+					const FString FeatureLevelName = LexToString(FeatureLevelType);
+					UE_LOG(LogIREEDriverRDG, Warning, TEXT("Shader platform %s: Minimum RHI Feature level is SM6 (desired feature level %s)!"), *ShaderPlatformName, *FeatureLevelName);
+					
+					continue;
+				}
 
-						continue;
-					}
+				if (FDataDrivenShaderPlatformInfo::GetSupportsRealTypes(ShaderPlatform) == ERHIFeatureSupport::Unsupported)
+				{
+					UE_LOG(LogIREEDriverRDG, Warning, TEXT("Shader platform %s does not support 16-bit types!"), *ShaderPlatformName);
 
-					if (FDataDrivenShaderPlatformInfo::GetSupportsRealTypes(ShaderPlatform) == ERHIFeatureSupport::Unsupported)
-					{
-						UE_LOG(LogIREEDriverRDG, Warning, TEXT("Shader platform %s does not support 16-bit types!"), *ShaderPlatformName);
+					continue;
+				}
 
-						continue;
-					}
+				for (const FString& HlslFilename : HlslFilenames)
+				{
+					UE_LOG(LogIREEDriverRDG, Log, TEXT("Process %s"), *HlslFilename);
 
-					for (const FString& HlslFilename : HlslFilenames)
+					if (!CompileAndSerializeShaderFromHLSLSource(HlslFilename, ShaderPlatform, TargetPlatform, IntermediateDirPath, true))
 					{
-						UE_LOG(LogIREEDriverRDG, Log, TEXT("Process %s"), *HlslFilename);
-
-						CompileAndSerializeShaderFromHLSLSource(HlslFilename, ShaderPlatform, TargetPlatform, OutDir, true);
+						bResult = false;
 					}
+				}
 
-					TArray<FString> ShaderPlatformShaderMapFilenames;
-					PlatformFile.FindFiles(ShaderPlatformShaderMapFilenames, *OutDir, TEXT(".ireeshadermap"));
-					
-					if (ShaderPlatformShaderMapFilenames.Num() != HlslFilenames.Num())
-					{
-						UE_LOG(LogIREEDriverRDG, Warning, TEXT("Not all shaders could be generated!"));
-
-						continue;
-					}
+				TArray<FString> ShaderPlatformShaderMapFilenames;
+				PlatformFile.FindFiles(ShaderPlatformShaderMapFilenames, *IntermediateDirPath, TEXT(".ireeshadermap"));
+				
+				if (ShaderPlatformShaderMapFilenames.Num() != HlslFilenames.Num())
+				{
+					UE_LOG(LogIREEDriverRDG, Warning, TEXT("Not all shaders could be generated!"));
+					bResult = false;
+					continue;
+				}
 
-					for (const FString& ShaderMapFilepath : ShaderPlatformShaderMapFilenames)
-					{
-						ShaderMapFilenames.Add(FPaths::Combine(ShaderPlatformName, FPaths::GetCleanFilename(ShaderMapFilepath)));
-					}
+				TArray<FString> ShaderMapFilenames;
+				for (const FString& ShaderMapFilepath : ShaderPlatformShaderMapFilenames)
+				{
+					ShaderMapFilenames.Add(FPaths::GetCleanFilename(ShaderMapFilepath));
 				}
 
 				FIREECompilerRDGBuildTargetResult ArchitectureInfo;
-				ArchitectureInfo.Architecture = BuildTargets[i].Architecture;
-				ArchitectureInfo.RelativeDirPath = BuildTargets[i].Architecture;
+				ArchitectureInfo.ShaderPlatform = Target->ShaderPlatform;
+				ArchitectureInfo.RelativeDirPath = Target->ShaderPlatform;
 				ArchitectureInfo.VmfbFileName = InModelName + ".vmfb";
 				ArchitectureInfo.SharedLibraryFileNames = ShaderMapFilenames;
 			
-				OutCompilerResult.ArchitectureInfos.Add(MoveTemp(ArchitectureInfo));
+				OutCompilerResult.BuildTargetResult.Add(MoveTemp(ArchitectureInfo));
 			}
 
-			bResult &= !OutCompilerResult.ArchitectureInfos.IsEmpty();
+			bResult &= !OutCompilerResult.BuildTargetResult.IsEmpty();
 			
 			if (!bResult)
 			{
-				OutCompilerResult.ArchitectureInfos.Empty();
+				OutCompilerResult.BuildTargetResult.Empty();
 			}
 			
 			return bResult;
diff --git a/Engine/Plugins/Experimental/NNERuntimeIREE/Source/IREEDriverRDG/Private/IREEDriverRDGExecutable.cpp b/Engine/Plugins/Experimental/NNERuntimeIREE/Source/IREEDriverRDG/Private/IREEDriverRDGExecutable.cpp
--- a/Engine/Plugins/Experimental/NNERuntimeIREE/Source/IREEDriverRDG/Private/IREEDriverRDGExecutable.cpp
+++ b/Engine/Plugins/Experimental/NNERuntimeIREE/Source/IREEDriverRDG/Private/IREEDriverRDGExecutable.cpp
@@ -66,12 +66,6 @@ public:
 		iree_hal_unreal_UnrealShaderDef_vec_t UnrealShaderVec = iree_hal_unreal_ExecutableDef_unreal_shaders_get(ExecutableDef);
 		iree_host_size_t UnrealShaderCount = iree_hal_unreal_UnrealShaderDef_vec_len(UnrealShaderVec);
 
-		const ERHIFeatureLevel::Type CacheFeatureLevel = GMaxRHIFeatureLevel;
-		const EShaderPlatform ShaderPlatform = GShaderPlatformForFeatureLevel[CacheFeatureLevel];
-		const FString ShaderPlatformName = LexToString(ShaderPlatform);
-
-		const FString ShaderMapSubDir = FPaths::Combine(BasePath, ShaderPlatformName);
-
 		for (iree_host_size_t i = 0; i < UnrealShaderCount; i++)
 		{
 			iree_hal_unreal_UnrealShaderDef_table_t UnrealShaderDef = iree_hal_unreal_UnrealShaderDef_vec_at(UnrealShaderVec, i);
@@ -81,7 +75,7 @@ public:
 		
 			const FString ShaderMapFileNameNoExt(ANSI_TO_TCHAR(SourceFileName));
 			const FString ShaderMapFileName = ShaderMapFileNameNoExt + TEXT(".ireeshadermap");
-			const FString ShaderMapFullPath = FPaths::Combine(ShaderMapSubDir, ShaderMapFileName);
+			const FString ShaderMapFullPath = FPaths::Combine(BasePath, ShaderMapFileName);
 
 			Executable->DebugShaderInfos.Add(ShaderMapFileNameNoExt);
 
@@ -101,7 +95,7 @@ public:
 
 			TUniquePtr<FNNERuntimeIREEResource> KernelResource = MakeUnique<FNNERuntimeIREEResource>();
 			KernelResource->SetupResource(
-				CacheFeatureLevel,
+				GMaxRHIFeatureLevel,
 				ShaderMapFileNameNoExt,
 				FString(),
 				FString(),
diff --git a/Engine/Plugins/Experimental/NNERuntimeIREE/Source/NNERuntimeIREE/Private/NNERuntimeIREE.cpp b/Engine/Plugins/Experimental/NNERuntimeIREE/Source/NNERuntimeIREE/Private/NNERuntimeIREE.cpp
--- a/Engine/Plugins/Experimental/NNERuntimeIREE/Source/NNERuntimeIREE/Private/NNERuntimeIREE.cpp
+++ b/Engine/Plugins/Experimental/NNERuntimeIREE/Source/NNERuntimeIREE/Private/NNERuntimeIREE.cpp
@@ -26,6 +26,7 @@
 #include "NNERuntimeIREEModel.h"
 #include "NNERuntimeIREEModelData.h"
 #include "RHIGlobals.h"
+#include "RHIStrings.h"
 #include "Serialization/MemoryReader.h"
 #include "Serialization/MemoryWriter.h"
 
@@ -74,12 +75,12 @@ namespace UE::NNERuntimeIREE::CPU::Private
 
 	FString GetIntermediateModelDirPath(const FString& PlatformName, const FString& RuntimeName, const FString& ModelName)
 	{
-		return FPaths::Combine("Intermediate", "Build", GetBinariesSubdirectory(PlatformName), UE_PLUGIN_NAME, RuntimeName, ModelName);
+		return FPaths::Combine("Intermediate", "Build", GetBinariesSubdirectory(PlatformName), RuntimeName, ModelName);
 	}
 
 	FString GetStagedModelDirPath(const FString& PlatformName, const FString& RuntimeName, const FString& ModelName)
 	{
-		return FPaths::Combine("Binaries", GetBinariesSubdirectory(PlatformName), UE_PLUGIN_NAME, RuntimeName, ModelName);
+		return FPaths::Combine("Binaries", GetBinariesSubdirectory(PlatformName), RuntimeName, ModelName);
 	}
 
 	FString GetPackagedModelDirPath(const FString& PlatformName, const FString& RuntimeName, const FString& ModelName)
@@ -546,7 +547,7 @@ UNNERuntimeIREERdg::ECanCreateModelDataStatus UNNERuntimeIREERdg::CanCreateModel
 		return ECanCreateModelDataStatus::Fail;
 	}
 
-	return	FileType.Compare(TEXT("mlir"), ESearchCase::IgnoreCase) == 0 ? ECanCreateModelDataStatus::Ok : ECanCreateModelDataStatus::FailFileIdNotSupported;
+	return	(FileType.Compare(TEXT("mlir"), ESearchCase::IgnoreCase) == 0 || FileType.Compare(TEXT("onnx"), ESearchCase::IgnoreCase) == 0) ? ECanCreateModelDataStatus::Ok : ECanCreateModelDataStatus::FailFileIdNotSupported;
 #else
 	return ECanCreateModelDataStatus::Fail;
 #endif // WITH_EDITOR
@@ -581,6 +582,25 @@ TSharedPtr<UE::NNE::FSharedModelData> UNNERuntimeIREERdg::CreateModelData(const
 	TWeakObjectPtr<UNNERuntimeIREEModelData> IREEModelData = NewObject<UNNERuntimeIREEModelData>();
 	FIREECompilerRDGResult CompilerResult{};
 
+	TArray<EShaderPlatform> ShaderPlatforms;
+	if (TargetPlatform)
+	{
+		TArray<FName> DesiredShaderFormats;
+		TargetPlatform->GetAllTargetedShaderFormats(DesiredShaderFormats);
+
+		for (const FName& ShaderFormatName : DesiredShaderFormats)
+		{
+			ShaderPlatforms.Add(ShaderFormatToLegacyShaderPlatform(ShaderFormatName));
+		}
+	}
+	else
+	{
+		const ERHIFeatureLevel::Type CacheFeatureLevel = GMaxRHIFeatureLevel;
+		const EShaderPlatform ShaderPlatform = GShaderPlatformForFeatureLevel[CacheFeatureLevel];
+
+		ShaderPlatforms.Add(ShaderPlatform);
+	}
+
 	bool bNeedCompileMlir = true;
 	if (PlatformFile.FileExists(*IREEModelDataFilePath))
 	{
@@ -601,16 +621,24 @@ TSharedPtr<UE::NNE::FSharedModelData> UNNERuntimeIREERdg::CreateModelData(const
 		}
 
 		bNeedCompileMlir = false;
-		for (int32 i = 0; i < CompilerResult.ArchitectureInfos.Num(); i++)
-		{
-			const FIREECompilerRDGBuildTargetResult& Info = CompilerResult.ArchitectureInfos[i];
-			const FString SharedLibrarySubDirFullPath = FPaths::Combine(SharedLibraryDirFullPath, Info.RelativeDirPath);
+		for (EShaderPlatform ShaderPlatform : ShaderPlatforms)
+		 {
+			const FString ShaderPlatformName = LexToString(ShaderPlatform);
 
-			const FString VmfbFilePath = FPaths::Combine(SharedLibrarySubDirFullPath, Info.VmfbFileName);
+			const FIREECompilerRDGBuildTargetResult* Info = CompilerResult.BuildTargetResult.FindByPredicate([ShaderPlatformName](const FIREECompilerRDGBuildTargetResult &Element) { return Element.ShaderPlatform == ShaderPlatformName; });
+			if (!Info)
+			{
+				bNeedCompileMlir = true;
+				continue;
+			}
+
+			const FString SharedLibrarySubDirFullPath = FPaths::Combine(SharedLibraryDirFullPath, Info->RelativeDirPath);
+
+			const FString VmfbFilePath = FPaths::Combine(SharedLibrarySubDirFullPath, Info->VmfbFileName);
 
 			bNeedCompileMlir |= !PlatformFile.FileExists(*VmfbFilePath);
 
-			for (const FString& ShaderMapFilename : Info.SharedLibraryFileNames)
+			for (const FString& ShaderMapFilename : Info->SharedLibraryFileNames)
 			{
 				const FString SharedLibraryFilePath = FPaths::Combine(SharedLibrarySubDirFullPath, ShaderMapFilename);
 
@@ -626,6 +654,8 @@ TSharedPtr<UE::NNE::FSharedModelData> UNNERuntimeIREERdg::CreateModelData(const
 		PlatformFile.DeleteDirectoryRecursively(*IntermediateDirFullPath);
 		PlatformFile.CreateDirectoryTree(*IntermediateDirFullPath);
 
+		CompilerResult = FIREECompilerRDGResult();
+
 		TUniquePtr<UE::IREE::Compiler::RDG::FCompiler> Compiler = UE::IREE::Compiler::RDG::FCompiler::Make(TargetPlatform);
 		if (!Compiler.IsValid())
 		{
@@ -633,7 +663,20 @@ TSharedPtr<UE::NNE::FSharedModelData> UNNERuntimeIREERdg::CreateModelData(const
 			return TSharedPtr<UE::NNE::FSharedModelData>();
 		}
 
-		if (!Compiler->CompileMlir(FileData, FileIdString, IntermediateDirFullPath, CompilerResult))
+		TArray64<uint8> ImporterOutputData;
+		if(FileType.Compare(TEXT("onnx"), ESearchCase::IgnoreCase) == 0)
+		{
+			if(!Compiler->ImportOnnx(FileData, FileIdString, IntermediateDirFullPath, ImporterOutputData))
+			{
+				UE_LOG(LogNNERuntimeIREE, Warning, TEXT("UNNERuntimeIREECpu failed to import ONNX model %s"), *FileIdString);
+				return TSharedPtr<UE::NNE::FSharedModelData>();
+			}
+			FileData = ImporterOutputData;
+		}
+
+		// NOTE: From this point FileData refers to mlir data in any case
+
+		if (!Compiler->CompileMlir(FileData, FileIdString, IntermediateDirFullPath, ShaderPlatforms, CompilerResult))
 		{
 			UE_LOG(LogNNERuntimeIREE, Warning, TEXT("UNNERuntimeIREERdg failed to compile model %s"), *FileIdString);
 			return TSharedPtr<UE::NNE::FSharedModelData>();
@@ -675,11 +718,11 @@ TSharedPtr<UE::NNE::FSharedModelData> UNNERuntimeIREERdg::CreateModelData(const
 	{
 		// Copy files for staging
 		FString StagingDirFullPath = FPaths::ConvertRelativePathToFull(FPaths::Combine(FPaths::ProjectDir(), GetPackagedModelDirPath(TargetPlatformName, GetRuntimeName(), FileIdString)));
-		for (int32 i = 0; i < CompilerResult.ArchitectureInfos.Num(); i++)
+		for (int32 i = 0; i < CompilerResult.BuildTargetResult.Num(); i++)
 		{
 			SCOPED_NAMED_EVENT_TEXT("Copy", FColor::Magenta);
 
-			const FIREECompilerRDGBuildTargetResult& Info = CompilerResult.ArchitectureInfos[i];
+			const FIREECompilerRDGBuildTargetResult& Info = CompilerResult.BuildTargetResult[i];
 			const FString SharedLibrarySubDirFullPath = FPaths::Combine(SharedLibraryDirFullPath, Info.RelativeDirPath);
 			const FString StagingSubDirFullPath = FPaths::Combine(StagingDirFullPath, Info.RelativeDirPath /* Info.Architecture */);
 			
@@ -746,13 +789,6 @@ TSharedPtr<UE::NNE::IModelRDG> UNNERuntimeIREERdg::CreateModelRDG(const TObjectP
 		return TSharedPtr<UE::NNE::IModelRDG>();
 	}
 
-	FString CurrentArchitecture = "";
-#if PLATFORM_CPU_X86_FAMILY
-	CurrentArchitecture = "x86_64";
-#elif PLATFORM_CPU_ARM_FAMILY
-	CurrentArchitecture = "arm64";
-#endif
-
 	TSharedPtr<UE::NNE::FSharedModelData> SharedData = ModelData->GetModelData(GetRuntimeName());
 	check(SharedData.IsValid());
 
@@ -787,30 +823,24 @@ TSharedPtr<UE::NNE::IModelRDG> UNNERuntimeIREERdg::CreateModelRDG(const TObjectP
 		FMemoryReaderView Reader(IREEModelData->CompilerResult, /*bIsPersitent =*/ true);
 		FIREECompilerRDGResult::StaticStruct()->SerializeBin(Reader, &CompilerResult);
 	}
+	
+	const ERHIFeatureLevel::Type CacheFeatureLevel = GMaxRHIFeatureLevel;
+	const EShaderPlatform ShaderPlatform = GShaderPlatformForFeatureLevel[CacheFeatureLevel];
+	const FString ShaderPlatformName = LexToString(ShaderPlatform);
 
-	int32 ArchitectureIndex = -1;
-	for (int32 i = 0; i < CompilerResult.ArchitectureInfos.Num(); i++)
-	{
-		if ((CompilerResult.ArchitectureInfos[i].Architecture.IsEmpty() && ArchitectureIndex < 0) || CompilerResult.ArchitectureInfos[i].Architecture.Equals(CurrentArchitecture))
-		{
-			ArchitectureIndex = i;
-		}
-	}
-
-	if (ArchitectureIndex < 0)
+	const FIREECompilerRDGBuildTargetResult* BuildTargetResult = CompilerResult.BuildTargetResult.FindByPredicate([ShaderPlatformName](const FIREECompilerRDGBuildTargetResult &Element) { return Element.ShaderPlatform == ShaderPlatformName; });
+	if (!BuildTargetResult)
 	{
-		UE_LOG(LogNNERuntimeIREE, Warning, TEXT("UNNERuntimeIREERdg failed to find a matching architecture for \'%s\'"), *CurrentArchitecture);
+		UE_LOG(LogNNERuntimeIREE, Warning, TEXT("UNNERuntimeIREERdg failed to find a matching shader platform for \'%s\'"), *ShaderPlatformName);
 		return {};
 	}
 
-	const FIREECompilerRDGBuildTargetResult& ArchitectureInfo = CompilerResult.ArchitectureInfos[ArchitectureIndex];
-
 	const FString FileIdString = IREEModelData->FileId.ToString(EGuidFormats::Digits).ToLower();
 	const FString RuntimeSubdir = GetRuntimeSubdir();
 	const FString SharedLibraryDirFullPath = FPaths::ConvertRelativePathToFull(FPaths::Combine(FPaths::ProjectDir(), GetSharedLibDirPath(UGameplayStatics::GetPlatformName(), GetRuntimeName(), FileIdString), RuntimeSubdir));
-	const FString SharedLibrarySubDirFullPath = FPaths::Combine(SharedLibraryDirFullPath, ArchitectureInfo.RelativeDirPath);
+	const FString SharedLibrarySubDirFullPath = FPaths::Combine(SharedLibraryDirFullPath, BuildTargetResult->RelativeDirPath);
 
-	TSharedPtr<UE::NNE::IModelRDG> Model = UE::NNERuntimeIREE::RDG::FModel::Make(SharedLibrarySubDirFullPath, ArchitectureInfo.VmfbFileName, *ModuleMetaData);
+	TSharedPtr<UE::NNE::IModelRDG> Model = UE::NNERuntimeIREE::RDG::FModel::Make(SharedLibrarySubDirFullPath, BuildTargetResult->VmfbFileName, *ModuleMetaData);
 	if (!Model.IsValid())
 	{
 		UE_LOG(LogNNERuntimeIREE, Warning, TEXT("UNNERuntimeIREERdg could not initialize the model created from model data with id %s"), *FileIdString);
