[NNE] Enable NNERuntimeIREERDG.

#tests All platforms, Editor & Game, ran model tests

#rnx

diff --git a/Engine/Plugins/Experimental/NNERuntimeIREE/Source/NNERuntimeIREE/NNERuntimeIREE.Build.cs b/Engine/Plugins/Experimental/NNERuntimeIREE/Source/NNERuntimeIREE/NNERuntimeIREE.Build.cs
--- a/Engine/Plugins/Experimental/NNERuntimeIREE/Source/NNERuntimeIREE/NNERuntimeIREE.Build.cs
+++ b/Engine/Plugins/Experimental/NNERuntimeIREE/Source/NNERuntimeIREE/NNERuntimeIREE.Build.cs
@@ -40,19 +40,19 @@ public class NNERuntimeIREE : ModuleRules
 		if (Target.Platform == UnrealTargetPlatform.Win64)
 		{
 			PublicDefinitions.Add("WITH_NNE_RUNTIME_IREE");
+			PublicDefinitions.Add("WITH_NNE_RUNTIME_IREE_RDG");
 		}
 		else if (Target.Platform == UnrealTargetPlatform.Linux)
 		{
 			PublicDefinitions.Add("WITH_NNE_RUNTIME_IREE");
 			PrivateDefinitions.Add("NNE_RUNTIME_IREE_USE_COMBINED_LIB_PATH");
+			PublicDefinitions.Add("WITH_NNE_RUNTIME_IREE_RDG");
 		}
 		else if (Target.Platform == UnrealTargetPlatform.Mac)
 		{
 			PublicDefinitions.Add("WITH_NNE_RUNTIME_IREE");
 			PrivateDefinitions.Add("NNE_RUNTIME_IREE_USE_COMBINED_LIB_PATH");
+			PublicDefinitions.Add("WITH_NNE_RUNTIME_IREE_RDG");
 		}
-
-		// Disable NNERuntimeIREERDG until development is progressed and runtime is functional
-		PublicDefinitions.Add("WITH_NNE_RUNTIME_IREE_RDG=false");
 	}
 }
diff --git a/Engine/Plugins/Experimental/NNERuntimeIREE/Source/NNERuntimeIREE/Private/NNERuntimeIREE.cpp b/Engine/Plugins/Experimental/NNERuntimeIREE/Source/NNERuntimeIREE/Private/NNERuntimeIREE.cpp
--- a/Engine/Plugins/Experimental/NNERuntimeIREE/Source/NNERuntimeIREE/Private/NNERuntimeIREE.cpp
+++ b/Engine/Plugins/Experimental/NNERuntimeIREE/Source/NNERuntimeIREE/Private/NNERuntimeIREE.cpp
@@ -530,6 +530,8 @@ int32 UNNERuntimeIREEVulkan::GetVersion() const
 	return Version;
 }
 
+#ifdef WITH_NNE_RUNTIME_IREE_RDG
+
 FGuid UNNERuntimeIREERdg::GUID = FGuid((int32)'I', (int32)'R', (int32)'D', (int32)'G');
 int32 UNNERuntimeIREERdg::Version = 0x00000001;
 
@@ -866,7 +868,9 @@ bool UNNERuntimeIREERdg::IsAvailable() const
 	}
 #endif
 
-	return WITH_NNE_RUNTIME_IREE_RDG;
+	return true;
 }
 
+#endif // WITH_NNE_RUNTIME_IREE_RDG
+
 #endif // WITH_NNE_RUNTIME_IREE
\ No newline at end of file
diff --git a/Engine/Plugins/Experimental/NNERuntimeIREE/Source/NNERuntimeIREE/Private/NNERuntimeIREE.h b/Engine/Plugins/Experimental/NNERuntimeIREE/Source/NNERuntimeIREE/Private/NNERuntimeIREE.h
--- a/Engine/Plugins/Experimental/NNERuntimeIREE/Source/NNERuntimeIREE/Private/NNERuntimeIREE.h
+++ b/Engine/Plugins/Experimental/NNERuntimeIREE/Source/NNERuntimeIREE/Private/NNERuntimeIREE.h
@@ -135,7 +135,7 @@ class UNNERuntimeIREERdg : public UObject, public INNERuntime, public INNERuntim
 	GENERATED_BODY()
 
 public:
-#ifdef WITH_NNE_RUNTIME_IREE
+#if defined(WITH_NNE_RUNTIME_IREE) && defined(WITH_NNE_RUNTIME_IREE_RDG)
 	static FGuid GUID;
 	static int32 Version;
 
diff --git a/Engine/Plugins/Experimental/NNERuntimeIREE/Source/NNERuntimeIREE/Private/NNERuntimeIREEModule.cpp b/Engine/Plugins/Experimental/NNERuntimeIREE/Source/NNERuntimeIREE/Private/NNERuntimeIREEModule.cpp
--- a/Engine/Plugins/Experimental/NNERuntimeIREE/Source/NNERuntimeIREE/Private/NNERuntimeIREEModule.cpp
+++ b/Engine/Plugins/Experimental/NNERuntimeIREE/Source/NNERuntimeIREE/Private/NNERuntimeIREEModule.cpp
@@ -60,6 +60,7 @@ void FNNERuntimeIREEModule::StartupModule()
 		}
 	}
 
+#ifdef WITH_NNE_RUNTIME_IREE_RDG
 	NNERuntimeIREERdg = NewObject<UNNERuntimeIREERdg>();
 	if (NNERuntimeIREERdg.IsValid())
 	{
@@ -73,6 +74,8 @@ void FNNERuntimeIREEModule::StartupModule()
 			NNERuntimeIREERdg.Reset();
 		}
 	}
+#endif // WITH_NNE_RUNTIME_IREE_RDG
+
 #endif // WITH_NNE_RUNTIME_IREE
 }
 
