[NNE] NNERuntimeIREECpu fix always compile MLIR when cooking.

#tests Win64, Editor & Game, verified re-compile on cook and ran model tests

#rnx

diff --git a/Engine/Plugins/Experimental/NNERuntimeIREE/Source/NNERuntimeIREE/Private/NNERuntimeIREE.cpp b/Engine/Plugins/Experimental/NNERuntimeIREE/Source/NNERuntimeIREE/Private/NNERuntimeIREE.cpp
--- a/Engine/Plugins/Experimental/NNERuntimeIREE/Source/NNERuntimeIREE/Private/NNERuntimeIREE.cpp
+++ b/Engine/Plugins/Experimental/NNERuntimeIREE/Source/NNERuntimeIREE/Private/NNERuntimeIREE.cpp
@@ -155,6 +155,8 @@ TSharedPtr<UE::NNE::FSharedModelData> UNNERuntimeIREECpu::CreateModelData(const
 
 	IPlatformFile& PlatformFile = FPlatformFileManager::Get().GetPlatformFile();
 
+	const bool bIsCooking = TargetPlatform != nullptr;
+
 	const FString FileIdString = FileId.ToString(EGuidFormats::Digits).ToLower();
 	const FString IntermediateDirFullPath = FPaths::ConvertRelativePathToFull(FPaths::Combine(FPaths::ProjectDir(), GetIntermediateModelDirPath(TargetPlatformName, GetRuntimeName(), FileIdString)));
 	const FString SharedLibraryDirFullPath = FPaths::ConvertRelativePathToFull(FPaths::Combine(FPaths::ProjectDir(), GetSharedLibDirPath(TargetPlatformName, GetRuntimeName(), FileIdString)));
@@ -198,7 +200,7 @@ TSharedPtr<UE::NNE::FSharedModelData> UNNERuntimeIREECpu::CreateModelData(const
 		}
 	}
 	
-	if (bNeedCompileMlir)
+	if (bNeedCompileMlir || bIsCooking)
 	{
 		SCOPED_NAMED_EVENT_TEXT("Compile", FColor::Magenta);
 
@@ -258,24 +260,28 @@ TSharedPtr<UE::NNE::FSharedModelData> UNNERuntimeIREECpu::CreateModelData(const
 		FFileHelper::SaveArrayToFile(ResultData, *IREEModelDataFilePath);
 	}
 
-	// Copy files for staging
-	FString StagingDirFullPath = FPaths::ConvertRelativePathToFull(FPaths::Combine(FPaths::ProjectDir(), GetPackagedModelDirPath(TargetPlatformName, GetRuntimeName(), FileIdString)));
-	for (int32 i = 0; i < CompilerResult.ArchitectureInfos.Num(); i++)
+	// Only stage when cooking
+	if (bIsCooking)
 	{
-		SCOPED_NAMED_EVENT_TEXT("Copy", FColor::Magenta);
+		// Copy files for staging
+		FString StagingDirFullPath = FPaths::ConvertRelativePathToFull(FPaths::Combine(FPaths::ProjectDir(), GetPackagedModelDirPath(TargetPlatformName, GetRuntimeName(), FileIdString)));
+		for (int32 i = 0; i < CompilerResult.ArchitectureInfos.Num(); i++)
+		{
+			SCOPED_NAMED_EVENT_TEXT("Copy", FColor::Magenta);
 
-		const FNNERuntimeIREEArchitectureInfoCPU& Info = CompilerResult.ArchitectureInfos[i];
-		const FString SharedLibrarySubDirFullPath = FPaths::Combine(SharedLibraryDirFullPath, Info.RelativeDirPath);
-		const FString StagingSubDirFullPath = FPaths::Combine(StagingDirFullPath, Info.Architecture);
-		
-		const FString SharedLibraryFilePathSrc = FPaths::Combine(SharedLibrarySubDirFullPath, Info.SharedLibraryFileName);
-		const FString VmfbFilePathSrc = FPaths::Combine(SharedLibrarySubDirFullPath, Info.VmfbFileName);
+			const FNNERuntimeIREEArchitectureInfoCPU& Info = CompilerResult.ArchitectureInfos[i];
+			const FString SharedLibrarySubDirFullPath = FPaths::Combine(SharedLibraryDirFullPath, Info.RelativeDirPath);
+			const FString StagingSubDirFullPath = FPaths::Combine(StagingDirFullPath, Info.Architecture);
+			
+			const FString SharedLibraryFilePathSrc = FPaths::Combine(SharedLibrarySubDirFullPath, Info.SharedLibraryFileName);
+			const FString VmfbFilePathSrc = FPaths::Combine(SharedLibrarySubDirFullPath, Info.VmfbFileName);
 
-		const FString SharedLibraryFilePathDest = FPaths::Combine(StagingSubDirFullPath, Info.SharedLibraryFileName);
-		const FString VmfbFilePathDest = FPaths::Combine(StagingSubDirFullPath, Info.VmfbFileName);
+			const FString SharedLibraryFilePathDest = FPaths::Combine(StagingSubDirFullPath, Info.SharedLibraryFileName);
+			const FString VmfbFilePathDest = FPaths::Combine(StagingSubDirFullPath, Info.VmfbFileName);
 
-		IFileManager::Get().Copy(*SharedLibraryFilePathDest, *SharedLibraryFilePathSrc, bNeedCompileMlir);
-		IFileManager::Get().Copy(*VmfbFilePathDest, *VmfbFilePathSrc, bNeedCompileMlir);
+			IFileManager::Get().Copy(*SharedLibraryFilePathDest, *SharedLibraryFilePathSrc);
+			IFileManager::Get().Copy(*VmfbFilePathDest, *VmfbFilePathSrc);
+		}
 	}
 
 	return MakeShared<UE::NNE::FSharedModelData>(MakeSharedBufferFromArray(MoveTemp(ResultData)), 0);
