[NNE] NNERuntimeIREERdg stage shaders as part of UNNEModelData.

#tests All platforms, Editor & Game, ran all rdg model tests

#rnx

diff --git a/Engine/Plugins/Experimental/NNERuntimeIREE/Config/Linux/LinuxGame.ini b/Engine/Plugins/Experimental/NNERuntimeIREE/Config/Linux/LinuxGame.ini
--- a/Engine/Plugins/Experimental/NNERuntimeIREE/Config/Linux/LinuxGame.ini
+++ b/Engine/Plugins/Experimental/NNERuntimeIREE/Config/Linux/LinuxGame.ini
@@ -1,3 +1,2 @@
 [/Script/UnrealEd.ProjectPackagingSettings]
 +DirectoriesToAlwaysStageAsNonUFS=(Path="../Binaries/Linux/NNERuntimeIREECpu")
-+DirectoriesToAlwaysStageAsNonUFS=(Path="../Binaries/Linux/NNERuntimeIREERdg")
diff --git a/Engine/Plugins/Experimental/NNERuntimeIREE/Config/LinuxArm64/LinuxArm64Game.ini b/Engine/Plugins/Experimental/NNERuntimeIREE/Config/LinuxArm64/LinuxArm64Game.ini
--- a/Engine/Plugins/Experimental/NNERuntimeIREE/Config/LinuxArm64/LinuxArm64Game.ini
+++ b/Engine/Plugins/Experimental/NNERuntimeIREE/Config/LinuxArm64/LinuxArm64Game.ini
@@ -1,3 +1,2 @@
 [/Script/UnrealEd.ProjectPackagingSettings]
 +DirectoriesToAlwaysStageAsNonUFS=(Path="../Binaries/LinuxArm64/NNERuntimeIREECpu")
-+DirectoriesToAlwaysStageAsNonUFS=(Path="../Binaries/LinuxArm64/NNERuntimeIREERdg")
diff --git a/Engine/Plugins/Experimental/NNERuntimeIREE/Config/Mac/MacGame.ini b/Engine/Plugins/Experimental/NNERuntimeIREE/Config/Mac/MacGame.ini
--- a/Engine/Plugins/Experimental/NNERuntimeIREE/Config/Mac/MacGame.ini
+++ b/Engine/Plugins/Experimental/NNERuntimeIREE/Config/Mac/MacGame.ini
@@ -1,3 +1,2 @@
 [/Script/UnrealEd.ProjectPackagingSettings]
 +DirectoriesToAlwaysStageAsNonUFS=(Path="../Binaries/Mac/NNERuntimeIREECpu")
-+DirectoriesToAlwaysStageAsNonUFS=(Path="../Binaries/Mac/NNERuntimeIREERdg")
diff --git a/Engine/Plugins/Experimental/NNERuntimeIREE/Config/Windows/WindowsGame.ini b/Engine/Plugins/Experimental/NNERuntimeIREE/Config/Windows/WindowsGame.ini
--- a/Engine/Plugins/Experimental/NNERuntimeIREE/Config/Windows/WindowsGame.ini
+++ b/Engine/Plugins/Experimental/NNERuntimeIREE/Config/Windows/WindowsGame.ini
@@ -1,3 +1,2 @@
 [/Script/UnrealEd.ProjectPackagingSettings]
 +DirectoriesToAlwaysStageAsNonUFS=(Path="../Binaries/Win64/NNERuntimeIREECpu")
-+DirectoriesToAlwaysStageAsNonUFS=(Path="../Binaries/Win64/NNERuntimeIREERdg")
diff --git a/Engine/Plugins/Experimental/NNERuntimeIREE/Source/IREEDriverRDG/Internal/IREECompilerRDG.h b/Engine/Plugins/Experimental/NNERuntimeIREE/Source/IREEDriverRDG/Internal/IREECompilerRDG.h
--- a/Engine/Plugins/Experimental/NNERuntimeIREE/Source/IREEDriverRDG/Internal/IREECompilerRDG.h
+++ b/Engine/Plugins/Experimental/NNERuntimeIREE/Source/IREEDriverRDG/Internal/IREECompilerRDG.h
@@ -9,21 +9,61 @@
 #include "IREECompilerRDG.generated.h"
 
 USTRUCT()
-struct FIREECompilerRDGBuildTargetResult
+struct FIREECompilerRDGExecutableData
 {
 	GENERATED_BODY()
 
 	UPROPERTY()
-	FString ShaderPlatform;
+	FString Name;
 
 	UPROPERTY()
-	FString RelativeDirPath;
-	
+	TArray<uint8> Data;
+};
+
+USTRUCT()
+struct FIREECompilerRDGBuildTargetResult
+{
+	GENERATED_BODY()
+
 	UPROPERTY()
-	TArray<FString> SharedLibraryFileNames;
+	FString ShaderPlatform;
 
 	UPROPERTY()
-	FString VmfbFileName;
+	TArray<FIREECompilerRDGExecutableData> Executables;
+
+	int64 DataSize;
+	TArray64<uint8> VmfbData;
+
+	bool Serialize(FArchive& Ar)
+	{
+		// Serialize normal UPROPERTY tagged data
+		UScriptStruct* Struct = FIREECompilerRDGBuildTargetResult::StaticStruct();
+		Struct->SerializeTaggedProperties(Ar, reinterpret_cast<uint8*>(this), Struct, nullptr);
+
+		if (Ar.IsLoading())
+		{
+			Ar << DataSize;
+			VmfbData.SetNumUninitialized(DataSize);
+			Ar.Serialize((void*)VmfbData.GetData(), DataSize);
+		}
+		else if (Ar.IsSaving())
+		{
+			DataSize = VmfbData.Num();
+			Ar << DataSize;
+			Ar.Serialize((void*)VmfbData.GetData(), DataSize);
+		}
+
+		return true;
+	}
+};
+
+template<>
+struct TStructOpsTypeTraits<FIREECompilerRDGBuildTargetResult> : public TStructOpsTypeTraitsBase2<FIREECompilerRDGBuildTargetResult>
+{
+	enum
+	{
+		WithSerializer = true
+	};
 };
 
 USTRUCT()
@@ -32,7 +72,7 @@ struct FIREECompilerRDGResult
 	GENERATED_BODY()
 
 	UPROPERTY()
-	TArray<FIREECompilerRDGBuildTargetResult> BuildTargetResult;
+	TArray<FIREECompilerRDGBuildTargetResult> BuildTargetResults;
 };
 
 #ifdef WITH_IREE_DRIVER_RDG
diff --git a/Engine/Plugins/Experimental/NNERuntimeIREE/Source/IREEDriverRDG/Internal/IREEDriverRDG.h b/Engine/Plugins/Experimental/NNERuntimeIREE/Source/IREEDriverRDG/Internal/IREEDriverRDG.h
--- a/Engine/Plugins/Experimental/NNERuntimeIREE/Source/IREEDriverRDG/Internal/IREEDriverRDG.h
+++ b/Engine/Plugins/Experimental/NNERuntimeIREE/Source/IREEDriverRDG/Internal/IREEDriverRDG.h
@@ -24,7 +24,7 @@ THIRD_PARTY_INCLUDES_END
 namespace UE::IREE::HAL::RDG
 {
 
-IREEDRIVERRDG_API iree_status_t DeviceCreate(iree_string_view_t Identifier, iree_allocator_t HostAllocator, const FString& BasePath, iree_hal_device_t** OutDevice);
+IREEDRIVERRDG_API iree_status_t DeviceCreate(iree_string_view_t Identifier, iree_allocator_t HostAllocator, const TMap<FString, TConstArrayView<uint8>>& Executables, iree_hal_device_t** OutDevice);
 
 IREEDRIVERRDG_API iree_status_t BufferWrapRDG(iree_allocator_t HostAllocator, iree_hal_allocator_t* DeviceAllocator, iree_hal_memory_type_t MemoryType, iree_hal_memory_access_t AllowedAccess, iree_hal_buffer_usage_t AllowedUsage, iree_device_size_t AllocationSize, iree_device_size_t ByteOffset, iree_device_size_t ByteLength, FRDGBuilder* GraphBuilder, FRDGBufferRef RDGBuffer, iree_hal_buffer_release_callback_t UserReleaseCallback, iree_hal_buffer_t** OutBuffer);
 IREEDRIVERRDG_API FRDGBufferRef BufferRDGBuffer(iree_hal_buffer_t* Buffer, FRDGBuilder* GraphBuilder);
diff --git a/Engine/Plugins/Experimental/NNERuntimeIREE/Source/IREEDriverRDG/Private/IREECompilerRDG.cpp b/Engine/Plugins/Experimental/NNERuntimeIREE/Source/IREEDriverRDG/Private/IREECompilerRDG.cpp
--- a/Engine/Plugins/Experimental/NNERuntimeIREE/Source/IREEDriverRDG/Private/IREECompilerRDG.cpp
+++ b/Engine/Plugins/Experimental/NNERuntimeIREE/Source/IREEDriverRDG/Private/IREECompilerRDG.cpp
@@ -115,8 +115,10 @@ namespace UE::IREE::Compiler
 				}
 			}
 
-			bool CompileAndSerializeShaderFromHLSLSource(const FString& HlslFilepath, EShaderPlatform ShaderPlatform, const ITargetPlatform* TargetPlatform, const FString& Outdir, bool bDumpHLSL)
+			bool CompileAndSerializeShaderFromHLSLSource(const FString& HlslFilepath, EShaderPlatform ShaderPlatform, const ITargetPlatform* TargetPlatform, const FString& Outdir, bool bDumpHLSL, FIREECompilerRDGExecutableData& ExecutableData)
 			{
+				UE_LOG(LogIREEDriverRDG, Log, TEXT("Process %s"), *HlslFilepath);
+
 				FString HlslSource;
 				if (!FFileHelper::LoadFileToString(HlslSource, *HlslFilepath))
 				{
@@ -169,7 +171,9 @@ namespace UE::IREE::Compiler
 					return false;
 				}
 
+#ifdef IREE_DEBUG_OUTPUT_SHADERMAP
 				const FString ShaderMapOutputFilepath = FPaths::Combine(Outdir, FPaths::GetBaseFilename(HlslFilepath, true) + ".ireeshadermap");
+#endif
 
 				TOptional<FArchiveSavePackageDataBuffer> ArchiveSavePackageData;
 
@@ -187,10 +191,40 @@ namespace UE::IREE::Compiler
 				// note: Ignore return value for now since errors should already be reported
 				/*bool bShaderMapValid = */ KernelResource->SerializeShaderMap(Writer);
 
+#ifdef IREE_DEBUG_OUTPUT_SHADERMAP
 				FFileHelper::SaveArrayToFile(ResultData, *ShaderMapOutputFilepath);
+#endif
+
+				ExecutableData.Name = FPaths::GetBaseFilename(HlslFilepath, true);
+				ExecutableData.Data = MoveTemp(ResultData);
 
 				return true;
 			}
+
+			TArray64<uint8> ReadVmfb(const FString& InVmfbPath)
+			{
+				check(!InVmfbPath.IsEmpty());
+			
+				TUniquePtr<FArchive> Reader = TUniquePtr<FArchive>(IFileManager::Get().CreateFileReader(*InVmfbPath, 0));
+				if (!Reader)
+				{
+					UE_LOG(LogIREEDriverRDG, Warning, TEXT("IREECompilerRDG failed to open the vmfb data file '%s'"), *InVmfbPath);
+					return {};
+				}
+				int64 DataSize = Reader->TotalSize();
+				if (DataSize < 1)
+				{
+					UE_LOG(LogIREEDriverRDG, Warning, TEXT("IREECompilerRDG vmfb data file '%s' is empty"), *InVmfbPath);
+					return {};
+				}
+
+				TArray64<uint8> Result;
+				Result.SetNumUninitialized(DataSize);
+
+				Reader->Serialize(Result.GetData(), DataSize);
+
+				return Result;
+			}
 		} // Private
 
 		FCompiler::FCompiler(const ITargetPlatform* InTargetPlatform, const FString& InImporterCommand, const FString& InImporterArguments, const FString& InCompilerCommand, const FString& InSharedLibExt, TConstArrayView<FBuildTarget> InBuildTargets) : TargetPlatform(InTargetPlatform), ImporterCommand(InImporterCommand), ImporterArguments(InImporterArguments), CompilerCommand(InCompilerCommand), BuildTargets(InBuildTargets)
@@ -368,9 +402,9 @@ namespace UE::IREE::Compiler
 
 		bool FCompiler::CompileMlir(TConstArrayView<uint8> InFileData, const FString& InModelName, const FString& InOutputDir, TConstArrayView<EShaderPlatform> ShaderPlatforms, FIREECompilerRDGResult& OutCompilerResult)
 		{
-			using namespace Private;
+			SCOPED_NAMED_EVENT_TEXT("FCompiler::CompileMlir", FColor::Magenta);
 
-			UE_LOG(LogIREEDriverRDG, Log, TEXT("FCompiler::CompileMlir %s"), *InModelName);
+			using namespace Private;
 
 			IPlatformFile& PlatformFile = FPlatformFileManager::Get().GetPlatformFile();
 
@@ -422,6 +456,14 @@ namespace UE::IREE::Compiler
 					continue;
 				}
 
+				TArray64<uint8> VmfbData = ReadVmfb(VmfbFilePath);
+				if (VmfbData.IsEmpty())
+				{
+					// Log in function...
+					bResult = false;
+					continue;
+				}
+
 				TArray<FString> HlslFilenames;
 				PlatformFile.FindFiles(HlslFilenames, *IntermediateDirPath, TEXT(".hlsl"));
 				if (HlslFilenames.IsEmpty())
@@ -447,46 +489,30 @@ namespace UE::IREE::Compiler
 					continue;
 				}
 
+				TArray<FIREECompilerRDGExecutableData> Executables;
 				for (const FString& HlslFilename : HlslFilenames)
 				{
-					UE_LOG(LogIREEDriverRDG, Log, TEXT("Process %s"), *HlslFilename);
-
-					if (!CompileAndSerializeShaderFromHLSLSource(HlslFilename, ShaderPlatform, TargetPlatform, IntermediateDirPath, true))
+					FIREECompilerRDGExecutableData& ExecutableData = Executables.Emplace_GetRef();
+					if (!CompileAndSerializeShaderFromHLSLSource(HlslFilename, ShaderPlatform, TargetPlatform, IntermediateDirPath, true, ExecutableData))
 					{
 						bResult = false;
 					}
 				}
 
-				TArray<FString> ShaderPlatformShaderMapFilenames;
-				PlatformFile.FindFiles(ShaderPlatformShaderMapFilenames, *IntermediateDirPath, TEXT(".ireeshadermap"));
-				
-				if (ShaderPlatformShaderMapFilenames.Num() != HlslFilenames.Num())
-				{
-					UE_LOG(LogIREEDriverRDG, Warning, TEXT("Not all shaders could be generated!"));
-					bResult = false;
-					continue;
-				}
-
-				TArray<FString> ShaderMapFilenames;
-				for (const FString& ShaderMapFilepath : ShaderPlatformShaderMapFilenames)
-				{
-					ShaderMapFilenames.Add(FPaths::GetCleanFilename(ShaderMapFilepath));
-				}
-
-				FIREECompilerRDGBuildTargetResult ArchitectureInfo;
-				ArchitectureInfo.ShaderPlatform = Target->ShaderPlatform;
-				ArchitectureInfo.RelativeDirPath = Target->ShaderPlatform;
-				ArchitectureInfo.VmfbFileName = InModelName + ".vmfb";
-				ArchitectureInfo.SharedLibraryFileNames = ShaderMapFilenames;
+				FIREECompilerRDGBuildTargetResult Result;
+				Result.ShaderPlatform = Target->ShaderPlatform;
+				Result.Executables = MoveTemp(Executables);
+				Result.DataSize = VmfbData.Num();
+				Result.VmfbData = MoveTemp(VmfbData);
 			
-				OutCompilerResult.BuildTargetResult.Add(MoveTemp(ArchitectureInfo));
+				OutCompilerResult.BuildTargetResults.Add(MoveTemp(Result));
 			}
 
-			bResult &= !OutCompilerResult.BuildTargetResult.IsEmpty();
+			bResult &= !OutCompilerResult.BuildTargetResults.IsEmpty();
 			
 			if (!bResult)
 			{
-				OutCompilerResult.BuildTargetResult.Empty();
+				OutCompilerResult.BuildTargetResults.Empty();
 			}
 			
 			return bResult;
diff --git a/Engine/Plugins/Experimental/NNERuntimeIREE/Source/IREEDriverRDG/Private/IREEDriverRDGDevice.cpp b/Engine/Plugins/Experimental/NNERuntimeIREE/Source/IREEDriverRDG/Private/IREEDriverRDGDevice.cpp
--- a/Engine/Plugins/Experimental/NNERuntimeIREE/Source/IREEDriverRDG/Private/IREEDriverRDGDevice.cpp
+++ b/Engine/Plugins/Experimental/NNERuntimeIREE/Source/IREEDriverRDG/Private/IREEDriverRDGDevice.cpp
@@ -4,6 +4,7 @@
 
 #ifdef WITH_IREE_DRIVER_RDG
 
+#include "Containers/Map.h"
 #include "IREEDriverRDGDeviceAllocator.h"
 #include "IREEDriverRDGDirectCommandBuffer.h"
 #include "IREEDriverRDGLog.h"
@@ -42,7 +43,7 @@ namespace Private
 class FDevice
 {
 public:
-	static iree_status_t Create(iree_string_view_t Identifier, iree_allocator_t HostAllocator, const FString& BasePath, iree_hal_device_t** OutDevice)
+	static iree_status_t Create(iree_string_view_t Identifier, iree_allocator_t HostAllocator, const TMap<FString, TConstArrayView<uint8>>& Executables, iree_hal_device_t** OutDevice)
 	{
 		check(OutDevice);
 
@@ -64,7 +65,11 @@ public:
 			return Status;
 		}
 
-		Device->BasePath = BasePath;
+		for (const auto& Pair : Executables)
+		{
+			Device->Executables.Emplace(Pair.Key, Pair.Value);
+		}
+
 		iree_arena_block_pool_initialize(32 * 1024, HostAllocator, &Device->BlockPool);
 
 		*OutDevice = (iree_hal_device_t*)Device;
@@ -234,7 +239,7 @@ private:
 #endif
 
 		FDevice* Device = Cast(BaseDevice);
-		return NoOpExecutableCacheCreate(Device->HostAllocator, Device->BasePath, OutExecutableCache);
+		return NoOpExecutableCacheCreate(Device->HostAllocator, &Device->Executables, OutExecutableCache);
 	}
 
 	static iree_status_t ImportFile(iree_hal_device_t* BaseDevice, iree_hal_queue_affinity_t QueueAffinity, iree_hal_memory_access_t Access, iree_io_file_handle_t *Handle, iree_hal_external_file_flags_t Flags, iree_hal_file_t **OutFile)
@@ -444,7 +449,7 @@ private:
 	iree_allocator_t HostAllocator;
 	iree_hal_allocator_t* DeviceAllocator;
 	iree_hal_channel_provider_t* ChannelProvider;
-	FString BasePath;
+	TMap<FString, TArray<uint8>> Executables;
 	// Block pool used for command buffers with a larger block size (as command
 	// buffers can contain inlined data uploads).
 	iree_arena_block_pool_t BlockPool;
@@ -484,9 +489,9 @@ const iree_hal_device_vtable_t FDevice::VTable =
 
 } // namespace Private
 
-iree_status_t DeviceCreate(iree_string_view_t Identifier, iree_allocator_t HostAllocator, const FString& BasePath, iree_hal_device_t** OutDevice)
+iree_status_t DeviceCreate(iree_string_view_t Identifier, iree_allocator_t HostAllocator, const TMap<FString, TConstArrayView<uint8>>& Executables, iree_hal_device_t** OutDevice)
 {
-	return Private::FDevice::Create(Identifier, HostAllocator, BasePath, OutDevice);
+	return Private::FDevice::Create(Identifier, HostAllocator, Executables, OutDevice);
 }
 
 } // UE::IREE
diff --git a/Engine/Plugins/Experimental/NNERuntimeIREE/Source/IREEDriverRDG/Private/IREEDriverRDGExecutable.cpp b/Engine/Plugins/Experimental/NNERuntimeIREE/Source/IREEDriverRDG/Private/IREEDriverRDGExecutable.cpp
--- a/Engine/Plugins/Experimental/NNERuntimeIREE/Source/IREEDriverRDG/Private/IREEDriverRDGExecutable.cpp
+++ b/Engine/Plugins/Experimental/NNERuntimeIREE/Source/IREEDriverRDG/Private/IREEDriverRDGExecutable.cpp
@@ -51,10 +51,11 @@ TArray<uint32> GetBufferBindings(const FIREEDriverRDGShaderParametersMetadata& M
 class FExecutable
 {
 public:
-	static iree_status_t Create(iree_allocator_t HostAllocator, const FString& BasePath, const iree_hal_executable_params_t* ExecutableParams, iree_hal_executable_t** OutExecutable)
+	static iree_status_t Create(iree_allocator_t HostAllocator, const TMap<FString, TArray<uint8>>* Executables, const iree_hal_executable_params_t* ExecutableParams, iree_hal_executable_t** OutExecutable)
 	{
 		SCOPED_NAMED_EVENT_TEXT("FExecutable::Create", FColor::Magenta);
 
+		check(Executables);
 		check(OutExecutable);
 
 		FExecutable* Executable;
@@ -74,18 +75,16 @@ public:
 			// flatbuffers_string_t EntryPoint = iree_hal_unreal_UnrealShaderDef_entry_point_get(UnrealShaderDef);
 		
 			const FString ShaderMapFileNameNoExt(ANSI_TO_TCHAR(SourceFileName));
-			const FString ShaderMapFileName = ShaderMapFileNameNoExt + TEXT(".ireeshadermap");
-			const FString ShaderMapFullPath = FPaths::Combine(BasePath, ShaderMapFileName);
 
-			Executable->DebugShaderInfos.Add(ShaderMapFileNameNoExt);
-
-			TArray/*64*/<uint8> Data;
-			if (!FFileHelper::LoadFileToArray(Data, *ShaderMapFullPath))
+			const TArray<uint8> *Data = Executables->Find(ShaderMapFileNameNoExt);
+			if (!Data)
 			{
-				return iree_make_status(IREE_STATUS_NOT_FOUND, "Could not load shader map.");
+				return iree_make_status(IREE_STATUS_NOT_FOUND, "Executable not found.");
 			}
 
-			FMemoryReaderView Reader(Data, /*bIsPersitent =*/ true);
+			Executable->DebugShaderInfos.Add(ShaderMapFileNameNoExt);
+
+			FMemoryReaderView Reader(*Data, /*bIsPersitent =*/ true);
 
 			FIREEDriverRDGShaderParametersMetadata IREEShaderParametersMetadata{};
 			FIREEDriverRDGShaderParametersMetadata::StaticStruct()->SerializeBin(Reader, &IREEShaderParametersMetadata);
@@ -170,9 +169,9 @@ private:
 
 } // namespace Private
 
-iree_status_t ExecutableCreate(iree_allocator_t HostAllocator, const FString& BasePath, const iree_hal_executable_params_t* ExecutableParams, iree_hal_executable_t** OutExecutable)
+iree_status_t ExecutableCreate(iree_allocator_t HostAllocator, const TMap<FString, TArray<uint8>>* Executables, const iree_hal_executable_params_t* ExecutableParams, iree_hal_executable_t** OutExecutable)
 {
-	return Private::FExecutable::Create(HostAllocator, BasePath, ExecutableParams, OutExecutable);
+	return Private::FExecutable::Create(HostAllocator, Executables, ExecutableParams, OutExecutable);
 }
 
 iree_status_t ExecutableGetResource(iree_hal_executable_t* Executable, int32 EntryPoint, const FNNERuntimeIREEResource** OutResource)
diff --git a/Engine/Plugins/Experimental/NNERuntimeIREE/Source/IREEDriverRDG/Private/IREEDriverRDGExecutable.h b/Engine/Plugins/Experimental/NNERuntimeIREE/Source/IREEDriverRDG/Private/IREEDriverRDGExecutable.h
--- a/Engine/Plugins/Experimental/NNERuntimeIREE/Source/IREEDriverRDG/Private/IREEDriverRDGExecutable.h
+++ b/Engine/Plugins/Experimental/NNERuntimeIREE/Source/IREEDriverRDG/Private/IREEDriverRDGExecutable.h
@@ -11,7 +11,7 @@ class FNNERuntimeIREEResource;
 namespace UE::IREE::HAL::RDG
 {
 
-iree_status_t ExecutableCreate(iree_allocator_t HostAllocator, const FString& BasePath, const iree_hal_executable_params_t* ExecutableParams, iree_hal_executable_t** OutExecutable);
+iree_status_t ExecutableCreate(iree_allocator_t HostAllocator, const TMap<FString, TArray<uint8>>* Executables, const iree_hal_executable_params_t* ExecutableParams, iree_hal_executable_t** OutExecutable);
 
 iree_status_t ExecutableGetResource(iree_hal_executable_t* Executable, int32 EntryPoint, const FNNERuntimeIREEResource** OutResource);
 
diff --git a/Engine/Plugins/Experimental/NNERuntimeIREE/Source/IREEDriverRDG/Private/IREEDriverRDGNoOpExecutableCache.cpp b/Engine/Plugins/Experimental/NNERuntimeIREE/Source/IREEDriverRDG/Private/IREEDriverRDGNoOpExecutableCache.cpp
--- a/Engine/Plugins/Experimental/NNERuntimeIREE/Source/IREEDriverRDG/Private/IREEDriverRDGNoOpExecutableCache.cpp
+++ b/Engine/Plugins/Experimental/NNERuntimeIREE/Source/IREEDriverRDG/Private/IREEDriverRDGNoOpExecutableCache.cpp
@@ -16,25 +16,21 @@ namespace Private
 class FNoOpExecutableCache
 {
 public:
-	static iree_status_t Create(iree_allocator_t HostAllocator, const FString& BasePath, iree_hal_executable_cache_t** OutExecutableCache)
+	static iree_status_t Create(iree_allocator_t HostAllocator, const TMap<FString, TArray<uint8>>* Executables, iree_hal_executable_cache_t** OutExecutableCache)
 	{
+		check(Executables);
 		check(OutExecutableCache);
 
 		FNoOpExecutableCache* ExecutableCache;
 		IREE_RETURN_IF_ERROR(iree_allocator_malloc(HostAllocator, sizeof(*ExecutableCache), (void**)&ExecutableCache));
 		iree_hal_resource_initialize((const void*)&FNoOpExecutableCache::VTable, &ExecutableCache->Resource);
 		ExecutableCache->HostAllocator = HostAllocator;
-		ExecutableCache->BasePath = BasePath;
+		ExecutableCache->Executables = Executables;
 
 		*OutExecutableCache = (iree_hal_executable_cache_t*)ExecutableCache;
 		return iree_ok_status();
 	}
 
-	FString GetBasePath() const
-	{
-		return BasePath;
-	}
-
 private:
 	static FNoOpExecutableCache* Cast(iree_hal_executable_cache_t* ExecutableCache)
 	{
@@ -79,14 +75,14 @@ private:
 
 		FNoOpExecutableCache* ExecutableCache = Cast(BaseExecutableCache);
 
-		return ExecutableCreate(ExecutableCache->HostAllocator, ExecutableCache->GetBasePath(), ExecutableParams, OutExecutable);
+		return ExecutableCreate(ExecutableCache->HostAllocator, ExecutableCache->Executables, ExecutableParams, OutExecutable);
 	}
 
 	static const iree_hal_executable_cache_vtable_t VTable;
 
 	iree_hal_resource_t Resource;
 	iree_allocator_t HostAllocator;
-	FString BasePath;
+	const TMap<FString, TArray<uint8>>* Executables;
 };
 
 const iree_hal_executable_cache_vtable_t FNoOpExecutableCache::VTable = 
@@ -98,9 +94,9 @@ const iree_hal_executable_cache_vtable_t FNoOpExecutableCache::VTable =
 
 } // namespace Private
 
-iree_status_t NoOpExecutableCacheCreate(iree_allocator_t HostAllocator, const FString& BasePath, iree_hal_executable_cache_t** OutExecutableCache)
+iree_status_t NoOpExecutableCacheCreate(iree_allocator_t HostAllocator, const TMap<FString, TArray<uint8>>* Executables, iree_hal_executable_cache_t** OutExecutableCache)
 {
-	return Private::FNoOpExecutableCache::Create(HostAllocator, BasePath, OutExecutableCache);
+	return Private::FNoOpExecutableCache::Create(HostAllocator, Executables, OutExecutableCache);
 }
 
 } // UE::IREE
diff --git a/Engine/Plugins/Experimental/NNERuntimeIREE/Source/IREEDriverRDG/Private/IREEDriverRDGNoOpExecutableCache.h b/Engine/Plugins/Experimental/NNERuntimeIREE/Source/IREEDriverRDG/Private/IREEDriverRDGNoOpExecutableCache.h
--- a/Engine/Plugins/Experimental/NNERuntimeIREE/Source/IREEDriverRDG/Private/IREEDriverRDGNoOpExecutableCache.h
+++ b/Engine/Plugins/Experimental/NNERuntimeIREE/Source/IREEDriverRDG/Private/IREEDriverRDGNoOpExecutableCache.h
@@ -9,7 +9,7 @@
 namespace UE::IREE::HAL::RDG
 {
 
-iree_status_t NoOpExecutableCacheCreate(iree_allocator_t HostAllocator, const FString& BasePath, iree_hal_executable_cache_t** OutExecutableCache);
+iree_status_t NoOpExecutableCacheCreate(iree_allocator_t HostAllocator, const TMap<FString, TArray<uint8>>* Executables, iree_hal_executable_cache_t** OutExecutableCache);
 
 } // UE::IREE
 
diff --git a/Engine/Plugins/Experimental/NNERuntimeIREE/Source/NNERuntimeIREE/Private/NNERuntimeIREE.cpp b/Engine/Plugins/Experimental/NNERuntimeIREE/Source/NNERuntimeIREE/Private/NNERuntimeIREE.cpp
--- a/Engine/Plugins/Experimental/NNERuntimeIREE/Source/NNERuntimeIREE/Private/NNERuntimeIREE.cpp
+++ b/Engine/Plugins/Experimental/NNERuntimeIREE/Source/NNERuntimeIREE/Private/NNERuntimeIREE.cpp
@@ -30,7 +30,7 @@
 #include "Serialization/MemoryReader.h"
 #include "Serialization/MemoryWriter.h"
 
-namespace UE::NNERuntimeIREE::CPU::Private
+namespace UE::NNERuntimeIREE::Private
 {
 	FString GetTargetPlatformName(const ITargetPlatform* TargetPlatform)
 	{
@@ -113,7 +113,31 @@ namespace UE::NNERuntimeIREE::CPU::Private
 		return {};
 #endif
 	}
-} // UE::NNERuntimeIREE::CPU::Private
+
+	TArray<EShaderPlatform> GetShaderPlatforms(const ITargetPlatform* TargetPlatform)
+	{
+		TArray<EShaderPlatform> ShaderPlatforms;
+		if (TargetPlatform)
+		{
+			TArray<FName> DesiredShaderFormats;
+			TargetPlatform->GetAllTargetedShaderFormats(DesiredShaderFormats);
+
+			for (const FName& ShaderFormatName : DesiredShaderFormats)
+			{
+				ShaderPlatforms.Add(ShaderFormatToLegacyShaderPlatform(ShaderFormatName));
+			}
+		}
+		else
+		{
+			const ERHIFeatureLevel::Type CacheFeatureLevel = GMaxRHIFeatureLevel;
+			const EShaderPlatform ShaderPlatform = GShaderPlatformForFeatureLevel[CacheFeatureLevel];
+
+			ShaderPlatforms.Add(ShaderPlatform);
+		}
+
+		return ShaderPlatforms;
+	}
+} // UE::NNERuntimeIREE::Private
 
 FGuid UNNERuntimeIREECpu::GUID = FGuid((int32)'I', (int32)'C', (int32)'P', (int32)'U');
 int32 UNNERuntimeIREECpu::Version = 0x00000005;
@@ -144,7 +168,7 @@ TSharedPtr<UE::NNE::FSharedModelData> UNNERuntimeIREECpu::CreateModelData(const
 	SCOPED_NAMED_EVENT_TEXT("UNNERuntimeIREECpu::CreateModelData", FColor::Magenta);
 
 #if WITH_EDITOR
-	using namespace UE::NNERuntimeIREE::CPU::Private;
+	using namespace UE::NNERuntimeIREE::Private;
 
 	const FString TargetPlatformName = GetTargetPlatformName(TargetPlatform);
 	if (CanCreateModelData(FileType, FileData, AdditionalFileData, FileId, TargetPlatform) != ECanCreateModelDataStatus::Ok)
@@ -164,7 +188,7 @@ TSharedPtr<UE::NNE::FSharedModelData> UNNERuntimeIREECpu::CreateModelData(const
 	FString IREEModelDataFilePath = FPaths::Combine(IntermediateDirFullPath, FileIdString) + ".ireemodeldata";
 
 	TArray64<uint8> ResultData;
-	TWeakObjectPtr<UNNERuntimeIREEModelData> IREEModelData = NewObject<UNNERuntimeIREEModelData>();
+	TWeakObjectPtr<UNNERuntimeIREEModelDataCPU> IREEModelData = NewObject<UNNERuntimeIREEModelDataCPU>();
 	FNNERuntimeIREECompilerResultCPU CompilerResult{};
 
 	bool bNeedCompileMlir = true;
@@ -293,8 +317,8 @@ TSharedPtr<UE::NNE::FSharedModelData> UNNERuntimeIREECpu::CreateModelData(const
 FString UNNERuntimeIREECpu::GetModelDataIdentifier(const FString& FileType, TConstArrayView64<uint8> FileData, const TMap<FString, TConstArrayView64<uint8>>& AdditionalFileData, const FGuid& FileId, const ITargetPlatform* TargetPlatform) const
 {
 	// Leave architecture blank as there is only one model data for all architectures of a given platform, only the vmfb and shared lib are different
-	FString PlatformName = UE::NNERuntimeIREE::CPU::Private::GetTargetPlatformName(TargetPlatform);
-	return UE::NNERuntimeIREE::CPU::Private::GetModelDataIdentifier(GetRuntimeName(), UNNERuntimeIREECpu::GUID, UNNERuntimeIREECpu::Version, FileId.ToString(EGuidFormats::Digits), PlatformName, "");
+	FString PlatformName = UE::NNERuntimeIREE::Private::GetTargetPlatformName(TargetPlatform);
+	return UE::NNERuntimeIREE::Private::GetModelDataIdentifier(GetRuntimeName(), UNNERuntimeIREECpu::GUID, UNNERuntimeIREECpu::Version, FileId.ToString(EGuidFormats::Digits), PlatformName, "");
 }
 
 UNNERuntimeIREECpu::ECanCreateModelCPUStatus UNNERuntimeIREECpu::CanCreateModelCPU(const TObjectPtr<UNNEModelData> ModelData) const
@@ -307,7 +331,7 @@ UNNERuntimeIREECpu::ECanCreateModelCPUStatus UNNERuntimeIREECpu::CanCreateModelC
 		return ECanCreateModelCPUStatus::Fail;
 	}
 
-	if (!UNNERuntimeIREEModelData::IsSameGuidAndVersion(SharedData->GetView(), UNNERuntimeIREECpu::GUID, UNNERuntimeIREECpu::Version))
+	if (!UNNERuntimeIREEModelDataCPU::IsSameGuidAndVersion(SharedData->GetView(), UNNERuntimeIREECpu::GUID, UNNERuntimeIREECpu::Version))
 	{
 		return ECanCreateModelCPUStatus::Fail;
 	}
@@ -319,7 +343,7 @@ TSharedPtr<UE::NNE::IModelCPU> UNNERuntimeIREECpu::CreateModelCPU(const TObjectP
 {
 	check(ModelData != nullptr);
 
-	using namespace UE::NNERuntimeIREE::CPU::Private;
+	using namespace UE::NNERuntimeIREE::Private;
 
 	if (CanCreateModelCPU(ModelData) != ECanCreateModelCPUStatus::Ok)
 	{
@@ -339,7 +363,7 @@ TSharedPtr<UE::NNE::IModelCPU> UNNERuntimeIREECpu::CreateModelCPU(const TObjectP
 
 	TConstArrayView<uint8> SharedDataView = SharedData->GetView();
 
-	TWeakObjectPtr<UNNERuntimeIREEModelData> IREEModelData = NewObject<UNNERuntimeIREEModelData>();
+	TWeakObjectPtr<UNNERuntimeIREEModelDataCPU> IREEModelData = NewObject<UNNERuntimeIREEModelDataCPU>();
 	{
 		FMemoryReaderView Reader(SharedDataView, /*bIsPersitent =*/ true);
 		IREEModelData->Serialize(Reader);
@@ -427,8 +451,8 @@ TSharedPtr<UE::NNE::FSharedModelData> UNNERuntimeIREEGpu::CreateModelData(const
 
 FString UNNERuntimeIREEGpu::GetModelDataIdentifier(const FString& FileType, TConstArrayView64<uint8> FileData, const TMap<FString, TConstArrayView64<uint8>>& AdditionalFileData, const FGuid& FileId, const ITargetPlatform* TargetPlatform) const
 {
-	FString PlatformName = UE::NNERuntimeIREE::CPU::Private::GetTargetPlatformName(TargetPlatform);
-	return UE::NNERuntimeIREE::CPU::Private::GetModelDataIdentifier(GetRuntimeName(), GetGUID(), UNNERuntimeIREEGpu::GetVersion(), FileId.ToString(EGuidFormats::Digits), PlatformName, "");
+	FString PlatformName = UE::NNERuntimeIREE::Private::GetTargetPlatformName(TargetPlatform);
+	return UE::NNERuntimeIREE::Private::GetModelDataIdentifier(GetRuntimeName(), GetGUID(), UNNERuntimeIREEGpu::GetVersion(), FileId.ToString(EGuidFormats::Digits), PlatformName, "");
 }
 
 UNNERuntimeIREEGpu::ECanCreateModelGPUStatus UNNERuntimeIREEGpu::CanCreateModelGPU(const TObjectPtr<UNNEModelData> ModelData) const
@@ -539,7 +563,7 @@ int32 UNNERuntimeIREEVulkan::GetVersion() const
 #ifdef WITH_NNE_RUNTIME_IREE_RDG
 
 FGuid UNNERuntimeIREERdg::GUID = FGuid((int32)'I', (int32)'R', (int32)'D', (int32)'G');
-int32 UNNERuntimeIREERdg::Version = 0x00000001;
+int32 UNNERuntimeIREERdg::Version = 0x00000002;
 
 FString UNNERuntimeIREERdg::GetRuntimeName() const
 {
@@ -566,7 +590,7 @@ TSharedPtr<UE::NNE::FSharedModelData> UNNERuntimeIREERdg::CreateModelData(const
 	SCOPED_NAMED_EVENT_TEXT("UNNERuntimeIREERdg::CreateModelData", FColor::Magenta);
 
 #if WITH_EDITOR
-	using namespace UE::NNERuntimeIREE::CPU::Private;
+	using namespace UE::NNERuntimeIREE::Private;
 
 	const FString TargetPlatformName = GetTargetPlatformName(TargetPlatform);
 	if (CanCreateModelData(FileType, FileData, AdditionalFileData, FileId, TargetPlatform) != ECanCreateModelDataStatus::Ok)
@@ -575,6 +599,8 @@ TSharedPtr<UE::NNE::FSharedModelData> UNNERuntimeIREERdg::CreateModelData(const
 		return TSharedPtr<UE::NNE::FSharedModelData>();
 	}
 
+	const TArray<EShaderPlatform> ShaderPlatforms = GetShaderPlatforms(TargetPlatform);
+
 	IPlatformFile& PlatformFile = FPlatformFileManager::Get().GetPlatformFile();
 
 	const bool bIsCooking = TargetPlatform != nullptr;
@@ -586,75 +612,33 @@ TSharedPtr<UE::NNE::FSharedModelData> UNNERuntimeIREERdg::CreateModelData(const
 
 	FString IREEModelDataFilePath = FPaths::Combine(IntermediateDirFullPath, FileIdString) + ".ireemodeldata";
 
-	TArray64<uint8> ResultData;
-	TWeakObjectPtr<UNNERuntimeIREEModelData> IREEModelData = NewObject<UNNERuntimeIREEModelData>();
-	FIREECompilerRDGResult CompilerResult{};
-
-	TArray<EShaderPlatform> ShaderPlatforms;
-	if (TargetPlatform)
-	{
-		TArray<FName> DesiredShaderFormats;
-		TargetPlatform->GetAllTargetedShaderFormats(DesiredShaderFormats);
-
-		for (const FName& ShaderFormatName : DesiredShaderFormats)
-		{
-			ShaderPlatforms.Add(ShaderFormatToLegacyShaderPlatform(ShaderFormatName));
-		}
-	}
-	else
-	{
-		const ERHIFeatureLevel::Type CacheFeatureLevel = GMaxRHIFeatureLevel;
-		const EShaderPlatform ShaderPlatform = GShaderPlatformForFeatureLevel[CacheFeatureLevel];
-
-		ShaderPlatforms.Add(ShaderPlatform);
-	}
-
 	bool bNeedCompileMlir = true;
 	if (PlatformFile.FileExists(*IREEModelDataFilePath))
 	{
 		SCOPED_NAMED_EVENT_TEXT("Validate", FColor::Magenta);
 
-		FFileHelper::LoadFileToArray(ResultData, *IREEModelDataFilePath);
-		
-		{
-			FMemoryReaderView Reader(ResultData, /*bIsPersitent =*/ true);
-			IREEModelData->Serialize(Reader);
-		}
-
-		check(FileIdString.Equals(IREEModelData->FileId.ToString(EGuidFormats::Digits).ToLower()));
-
+		TUniquePtr<FArchive> FileReader(IFileManager::Get().CreateFileReader(*IREEModelDataFilePath));
+		if (FileReader)
 		{
-			FMemoryReaderView Reader(IREEModelData->CompilerResult, /*bIsPersitent =*/ true);
-			FIREECompilerRDGResult::StaticStruct()->SerializeBin(Reader, &CompilerResult);
-		}
-
-		bNeedCompileMlir = false;
-		for (EShaderPlatform ShaderPlatform : ShaderPlatforms)
-		 {
-			const FString ShaderPlatformName = LexToString(ShaderPlatform);
-
-			const FIREECompilerRDGBuildTargetResult* Info = CompilerResult.BuildTargetResult.FindByPredicate([ShaderPlatformName](const FIREECompilerRDGBuildTargetResult &Element) { return Element.ShaderPlatform == ShaderPlatformName; });
-			if (!Info)
-			{
-				bNeedCompileMlir = true;
-				continue;
-			}
+			bNeedCompileMlir = false;
 
-			const FString SharedLibrarySubDirFullPath = FPaths::Combine(SharedLibraryDirFullPath, Info->RelativeDirPath);
+			FNNERuntimeIREEModelDataHeaderRDG Header = UNNERuntimeIREEModelDataRDG::ReadHeader(*FileReader);
 
-			const FString VmfbFilePath = FPaths::Combine(SharedLibrarySubDirFullPath, Info->VmfbFileName);
+			check(FileIdString.Equals(Header.FileId.ToString(EGuidFormats::Digits).ToLower()));
 
-			bNeedCompileMlir |= !PlatformFile.FileExists(*VmfbFilePath);
+			bNeedCompileMlir |= Header.GUID != GUID;
+			bNeedCompileMlir |= Header.Version != Version;
 
-			for (const FString& ShaderMapFilename : Info->SharedLibraryFileNames)
+			for (EShaderPlatform ShaderPlatform : ShaderPlatforms)
 			{
-				const FString SharedLibraryFilePath = FPaths::Combine(SharedLibrarySubDirFullPath, ShaderMapFilename);
+				const FString ShaderPlatformName = LexToString(ShaderPlatform);
 
-				bNeedCompileMlir |= !PlatformFile.FileExists(*SharedLibraryFilePath);
+				bNeedCompileMlir |= Header.ShaderPlatforms.Find(ShaderPlatformName) == INDEX_NONE;
 			}
 		}
 	}
-	
+
+	TArray64<uint8> ResultData;
 	if (bNeedCompileMlir || bIsCooking)
 	{
 		SCOPED_NAMED_EVENT_TEXT("Compile", FColor::Magenta);
@@ -662,7 +646,7 @@ TSharedPtr<UE::NNE::FSharedModelData> UNNERuntimeIREERdg::CreateModelData(const
 		PlatformFile.DeleteDirectoryRecursively(*IntermediateDirFullPath);
 		PlatformFile.CreateDirectoryTree(*IntermediateDirFullPath);
 
-		CompilerResult = FIREECompilerRDGResult();
+		FIREECompilerRDGResult CompilerResult{};
 
 		TUniquePtr<UE::IREE::Compiler::RDG::FCompiler> Compiler = UE::IREE::Compiler::RDG::FCompiler::Make(TargetPlatform);
 		if (!Compiler.IsValid())
@@ -690,9 +674,16 @@ TSharedPtr<UE::NNE::FSharedModelData> UNNERuntimeIREERdg::CreateModelData(const
 			return TSharedPtr<UE::NNE::FSharedModelData>();
 		}
 
-		IREEModelData->GUID = UNNERuntimeIREERdg::GUID;
-		IREEModelData->Version = UNNERuntimeIREERdg::Version;
-		IREEModelData->FileId = FileId;
+		TWeakObjectPtr<UNNERuntimeIREEModelDataRDG> IREEModelData = NewObject<UNNERuntimeIREEModelDataRDG>();
+		IREEModelData->Header.GUID = UNNERuntimeIREERdg::GUID;
+		IREEModelData->Header.Version = UNNERuntimeIREERdg::Version;
+		IREEModelData->Header.FileId = FileId;
+		for (EShaderPlatform ShaderPlatform : ShaderPlatforms)
+		{
+			const FString ShaderPlatformName = LexToString(ShaderPlatform);
+
+			IREEModelData->Header.ShaderPlatforms.Add(ShaderPlatformName);
+		}
 		if (AdditionalFileData.Contains("IREEModuleMetaData"))
 		{
 			IREEModelData->ModuleMetaData = AdditionalFileData["IREEModuleMetaData"];
@@ -718,47 +709,31 @@ TSharedPtr<UE::NNE::FSharedModelData> UNNERuntimeIREERdg::CreateModelData(const
 			IREEModelData->Serialize(Writer);
 		}
 
-		FFileHelper::SaveArrayToFile(ResultData, *IREEModelDataFilePath);
+		if (!FFileHelper::SaveArrayToFile(ResultData, *IREEModelDataFilePath))
+		{
+			UE_LOG(LogNNERuntimeIREE, Warning, TEXT("UNNERuntimeIREERdg failed to write model data to disk %s"), *IREEModelDataFilePath);
+			return {};
+		}
 	}
-
-	// Only stage when cooking
-	if (bIsCooking)
+	else
 	{
-		// Copy files for staging
-		FString StagingDirFullPath = FPaths::ConvertRelativePathToFull(FPaths::Combine(FPaths::ProjectDir(), GetPackagedModelDirPath(TargetPlatformName, GetRuntimeName(), FileIdString)));
-		for (int32 i = 0; i < CompilerResult.BuildTargetResult.Num(); i++)
+		if (!FFileHelper::LoadFileToArray(ResultData, *IREEModelDataFilePath))
 		{
-			SCOPED_NAMED_EVENT_TEXT("Copy", FColor::Magenta);
-
-			const FIREECompilerRDGBuildTargetResult& Info = CompilerResult.BuildTargetResult[i];
-			const FString SharedLibrarySubDirFullPath = FPaths::Combine(SharedLibraryDirFullPath, Info.RelativeDirPath);
-			const FString StagingSubDirFullPath = FPaths::Combine(StagingDirFullPath, Info.RelativeDirPath /* Info.Architecture */);
-			
-			const FString VmfbFilePathSrc = FPaths::Combine(SharedLibrarySubDirFullPath, Info.VmfbFileName);
-			const FString VmfbFilePathDest = FPaths::Combine(StagingSubDirFullPath, Info.VmfbFileName);
-
-			IFileManager::Get().Copy(*VmfbFilePathDest, *VmfbFilePathSrc);
-
-			for (const FString& ShaderMapFilename : Info.SharedLibraryFileNames)
-			{
-				const FString SharedLibraryFilePathSrc = FPaths::Combine(SharedLibrarySubDirFullPath, ShaderMapFilename);
-				const FString SharedLibraryFilePathDest = FPaths::Combine(StagingSubDirFullPath, ShaderMapFilename);
-
-				IFileManager::Get().Copy(*SharedLibraryFilePathDest, *SharedLibraryFilePathSrc);
-			}
+			UE_LOG(LogNNERuntimeIREE, Warning, TEXT("UNNERuntimeIREERdg failed to read model data from disk %s"), *IREEModelDataFilePath);
+			return {};
 		}
 	}
 
 	return MakeShared<UE::NNE::FSharedModelData>(MakeSharedBufferFromArray(MoveTemp(ResultData)), 0);
 #else
-	return TSharedPtr<UE::NNE::FSharedModelData>();
+	return {};
 #endif // WITH_EDITOR
 }
 
 FString UNNERuntimeIREERdg::GetModelDataIdentifier(const FString& FileType, TConstArrayView64<uint8> FileData, const TMap<FString, TConstArrayView64<uint8>>& AdditionalFileData, const FGuid& FileId, const ITargetPlatform* TargetPlatform) const
 {
-	FString PlatformName = UE::NNERuntimeIREE::CPU::Private::GetTargetPlatformName(TargetPlatform);
-	return UE::NNERuntimeIREE::CPU::Private::GetModelDataIdentifier(GetRuntimeName(), UNNERuntimeIREERdg::GUID, UNNERuntimeIREERdg::Version, FileId.ToString(EGuidFormats::Digits), PlatformName, "");
+	FString PlatformName = UE::NNERuntimeIREE::Private::GetTargetPlatformName(TargetPlatform);
+	return UE::NNERuntimeIREE::Private::GetModelDataIdentifier(GetRuntimeName(), UNNERuntimeIREERdg::GUID, UNNERuntimeIREERdg::Version, FileId.ToString(EGuidFormats::Digits), PlatformName, "");
 }
 
 UNNERuntimeIREERdg::ECanCreateModelRDGStatus UNNERuntimeIREERdg::CanCreateModelRDG(const TObjectPtr<UNNEModelData> ModelData) const
@@ -771,25 +746,32 @@ UNNERuntimeIREERdg::ECanCreateModelRDGStatus UNNERuntimeIREERdg::CanCreateModelR
 		return ECanCreateModelRDGStatus::Fail;
 	}
 
-	TConstArrayView<uint8> SharedDataView = SharedData->GetView();
-	int32 GuidSize = sizeof(UNNERuntimeIREERdg::GUID);
-	int32 VersionSize = sizeof(UNNERuntimeIREERdg::Version);
-	if (SharedDataView.Num() <= GuidSize + VersionSize)
+	FMemoryReaderView Reader(SharedData->GetView(), /*bIsPersitent =*/ true);
+	FNNERuntimeIREEModelDataHeaderRDG Header = UNNERuntimeIREEModelDataRDG::ReadHeader(Reader);
+
+	if (Header.GUID != GUID)
+	{
+		return ECanCreateModelRDGStatus::Fail;
+	}
+	if (Header.Version != Version)
 	{
 		return ECanCreateModelRDGStatus::Fail;
 	}
 
-	bool bResult = FGenericPlatformMemory::Memcmp(&(SharedDataView[0]), &(UNNERuntimeIREERdg::GUID), GuidSize) == 0;
-	bResult &= FGenericPlatformMemory::Memcmp(&(SharedDataView[GuidSize]), &(UNNERuntimeIREERdg::Version), VersionSize) == 0;
+	const FString ShaderPlatformName = LexToString(GShaderPlatformForFeatureLevel[GMaxRHIFeatureLevel]);
+	if (Header.ShaderPlatforms.Find(ShaderPlatformName) == INDEX_NONE)
+	{
+		return ECanCreateModelRDGStatus::Fail;
+	}
 
-	return bResult ? ECanCreateModelRDGStatus::Ok : ECanCreateModelRDGStatus::Fail;
+	return ECanCreateModelRDGStatus::Ok;
 }
 
 TSharedPtr<UE::NNE::IModelRDG> UNNERuntimeIREERdg::CreateModelRDG(const TObjectPtr<UNNEModelData> ModelData)
 {
 	check(ModelData != nullptr);
 
-	using namespace UE::NNERuntimeIREE::CPU::Private;
+	using namespace UE::NNERuntimeIREE::Private;
 
 	if (CanCreateModelRDG(ModelData) != ECanCreateModelRDGStatus::Ok)
 	{
@@ -802,7 +784,7 @@ TSharedPtr<UE::NNE::IModelRDG> UNNERuntimeIREERdg::CreateModelRDG(const TObjectP
 
 	TConstArrayView<uint8> SharedDataView = SharedData->GetView();
 
-	TWeakObjectPtr<UNNERuntimeIREEModelData> IREEModelData = NewObject<UNNERuntimeIREEModelData>();
+	TWeakObjectPtr<UNNERuntimeIREEModelDataRDG> IREEModelData = NewObject<UNNERuntimeIREEModelDataRDG>();
 	{
 		FMemoryReaderView Reader(SharedDataView, /*bIsPersitent =*/ true);
 		IREEModelData->Serialize(Reader);
@@ -832,23 +814,38 @@ TSharedPtr<UE::NNE::IModelRDG> UNNERuntimeIREERdg::CreateModelRDG(const TObjectP
 		FIREECompilerRDGResult::StaticStruct()->SerializeBin(Reader, &CompilerResult);
 	}
 	
-	const ERHIFeatureLevel::Type CacheFeatureLevel = GMaxRHIFeatureLevel;
-	const EShaderPlatform ShaderPlatform = GShaderPlatformForFeatureLevel[CacheFeatureLevel];
-	const FString ShaderPlatformName = LexToString(ShaderPlatform);
-
-	const FIREECompilerRDGBuildTargetResult* BuildTargetResult = CompilerResult.BuildTargetResult.FindByPredicate([ShaderPlatformName](const FIREECompilerRDGBuildTargetResult &Element) { return Element.ShaderPlatform == ShaderPlatformName; });
+	const FString ShaderPlatformName = LexToString(GShaderPlatformForFeatureLevel[GMaxRHIFeatureLevel]);
+	const FIREECompilerRDGBuildTargetResult* BuildTargetResult = CompilerResult.BuildTargetResults.FindByPredicate([ShaderPlatformName](const FIREECompilerRDGBuildTargetResult &Element) { return Element.ShaderPlatform == ShaderPlatformName; });
 	if (!BuildTargetResult)
 	{
 		UE_LOG(LogNNERuntimeIREE, Warning, TEXT("UNNERuntimeIREERdg failed to find a matching shader platform for \'%s\'"), *ShaderPlatformName);
 		return {};
 	}
 
-	const FString FileIdString = IREEModelData->FileId.ToString(EGuidFormats::Digits).ToLower();
+	TMap<FString, TConstArrayView<uint8>> ExecutableMap;
+	for (const FIREECompilerRDGExecutableData& ExecutableData : BuildTargetResult->Executables)
+	{
+		if (ExecutableData.Name.IsEmpty())
+		{
+			UE_LOG(LogNNERuntimeIREE, Warning, TEXT("UNNERuntimeIREERdg failed to load executable map, key is empty"));
+			return {};
+		}
+
+		if (ExecutableData.Data.IsEmpty())
+		{
+			UE_LOG(LogNNERuntimeIREE, Warning, TEXT("UNNERuntimeIREERdg failed to load executable map, value for key %s is empty"), *ExecutableData.Name);
+			return {};
+		}
+
+		ExecutableMap.Emplace(ExecutableData.Name, ExecutableData.Data);
+	}
+
+	const FString FileIdString = IREEModelData->Header.FileId.ToString(EGuidFormats::Digits).ToLower();
 	const FString RuntimeSubdir = GetRuntimeSubdir();
 	const FString SharedLibraryDirFullPath = FPaths::ConvertRelativePathToFull(FPaths::Combine(FPaths::ProjectDir(), GetSharedLibDirPath(UGameplayStatics::GetPlatformName(), GetRuntimeName(), FileIdString), RuntimeSubdir));
-	const FString SharedLibrarySubDirFullPath = FPaths::Combine(SharedLibraryDirFullPath, BuildTargetResult->RelativeDirPath);
+	const FString SharedLibrarySubDirFullPath = FPaths::Combine(SharedLibraryDirFullPath, BuildTargetResult->ShaderPlatform);
 
-	TSharedPtr<UE::NNE::IModelRDG> Model = UE::NNERuntimeIREE::RDG::FModel::Make(SharedLibrarySubDirFullPath, BuildTargetResult->VmfbFileName, *ModuleMetaData);
+	TSharedPtr<UE::NNE::IModelRDG> Model = UE::NNERuntimeIREE::RDG::FModel::Make(SharedLibrarySubDirFullPath, BuildTargetResult->VmfbData, *ModuleMetaData, ExecutableMap);
 	if (!Model.IsValid())
 	{
 		UE_LOG(LogNNERuntimeIREE, Warning, TEXT("UNNERuntimeIREERdg could not initialize the model created from model data with id %s"), *FileIdString);
diff --git a/Engine/Plugins/Experimental/NNERuntimeIREE/Source/NNERuntimeIREE/Private/NNERuntimeIREEModel.cpp b/Engine/Plugins/Experimental/NNERuntimeIREE/Source/NNERuntimeIREE/Private/NNERuntimeIREEModel.cpp
--- a/Engine/Plugins/Experimental/NNERuntimeIREE/Source/NNERuntimeIREE/Private/NNERuntimeIREEModel.cpp
+++ b/Engine/Plugins/Experimental/NNERuntimeIREE/Source/NNERuntimeIREE/Private/NNERuntimeIREEModel.cpp
@@ -388,36 +388,37 @@ namespace UE::NNERuntimeIREE
 				iree_vm_module_release(Module);
 			}
 
-			static TSharedPtr<FModule> Make(const FString& InVmfbPath, const FString& InVmfbName, const UNNERuntimeIREEModuleMetaData& InModuleMetaData)
+			static TSharedPtr<UE::NNE::FSharedModelData> GetVmfbAsModelData(const FString& InVmfbPath)
 			{
-				check(!InVmfbName.IsEmpty());
-				check(!InModuleMetaData.FunctionMetaData.IsEmpty());
-
-				FString CombinedPath = FPaths::Combine(InVmfbPath, InVmfbName);
-				if (Modules.Contains(CombinedPath) && Modules[CombinedPath].IsValid())
-				{
-					return Modules[CombinedPath].Pin();
-				}
+				check(!InVmfbPath.IsEmpty());
 
-				TUniquePtr<FArchive> Reader = TUniquePtr<FArchive>(IFileManager::Get().CreateFileReader(*CombinedPath, 0));
+				TUniquePtr<FArchive> Reader = TUniquePtr<FArchive>(IFileManager::Get().CreateFileReader(*InVmfbPath));
 				if (!Reader)
 				{
-					UE_LOG(LogNNERuntimeIREE, Error, TEXT("CPU module: Failed to open the vmfb data file '%s'"), *CombinedPath);
-					return TSharedPtr<FModule>();
+					UE_LOG(LogNNERuntimeIREE, Error, TEXT("CPU module: Failed to open the vmfb data file '%s'"), *InVmfbPath);
+					return {};
 				}
 				int64 DataSize = Reader->TotalSize();
 				if (DataSize < 1)
 				{
-					UE_LOG(LogNNERuntimeIREE, Error, TEXT("CPU module: vmfb data file '%s' is empty"), *CombinedPath);
-					return TSharedPtr<FModule>();
+					UE_LOG(LogNNERuntimeIREE, Error, TEXT("CPU module: vmfb data file '%s' is empty"), *InVmfbPath);
+					return {};
 				}
 
 				void* Data = FMemory::Malloc(DataSize, IREE_HAL_HEAP_BUFFER_ALIGNMENT);
 				Reader->Serialize(Data, DataSize);
-				TSharedPtr<UE::NNE::FSharedModelData> ModelData = MakeShared<UE::NNE::FSharedModelData>(FSharedBuffer::TakeOwnership(Data, DataSize, FMemory::Free), IREE_HAL_HEAP_BUFFER_ALIGNMENT);
-				if (!ModelData.IsValid())
+
+				return MakeShared<UE::NNE::FSharedModelData>(FSharedBuffer::TakeOwnership(Data, DataSize, FMemory::Free), IREE_HAL_HEAP_BUFFER_ALIGNMENT);
+			}
+
+			static TSharedPtr<FModule> Make(const FString& InVmfbPath, TSharedRef<UE::NNE::FSharedModelData> ModelData, const UNNERuntimeIREEModuleMetaData& InModuleMetaData)
+			{
+				check(!InVmfbPath.IsEmpty());
+				check(!InModuleMetaData.FunctionMetaData.IsEmpty());
+
+				if (Modules.Contains(InVmfbPath) && Modules[InVmfbPath].IsValid())
 				{
-					return TSharedPtr<FModule>();
+					return Modules[InVmfbPath].Pin();
 				}
 
 				TSharedPtr<FInstance> Instance = FInstance::GetInstance();
@@ -432,8 +433,8 @@ namespace UE::NNERuntimeIREE
 					return TSharedPtr<FModule>();
 				}
 
-				TSharedPtr<FModule> Result = TSharedPtr<FModule>(new FModule(Instance.ToSharedRef(), ModelData.ToSharedRef(), Module, InModuleMetaData.FunctionMetaData));
-				Modules.Add(CombinedPath, Result);
+				TSharedPtr<FModule> Result = TSharedPtr<FModule>(new FModule(Instance.ToSharedRef(), ModelData, Module, InModuleMetaData.FunctionMetaData));
+				Modules.Add(InVmfbPath, Result);
 				return Result;
 			}
 
@@ -1062,10 +1063,17 @@ namespace UE::NNERuntimeIREE
 				return TSharedPtr<FModel>();
 			}
 
-			TSharedPtr<UE::NNERuntimeIREE::Private::FModule> Module = UE::NNERuntimeIREE::Private::FModule::Make(InDirPath, InVmfbFileName, InModuleMetaData);
+			const FString CombinedPath = FPaths::Combine(InDirPath, InVmfbFileName);
+			const TSharedPtr<UE::NNE::FSharedModelData> ModelData = UE::NNERuntimeIREE::Private::FModule::GetVmfbAsModelData(CombinedPath);
+			if (!ModelData.IsValid())
+			{
+				return {};
+			}
+
+			TSharedPtr<UE::NNERuntimeIREE::Private::FModule> Module = UE::NNERuntimeIREE::Private::FModule::Make(CombinedPath, ModelData.ToSharedRef(), InModuleMetaData);
 			if (!Module.IsValid())
 			{
-				return TSharedPtr<FModel>();
+				return {};
 			}
 
 			return TSharedPtr<FModel>(new FModel(Device.ToSharedRef(), Module.ToSharedRef()));
@@ -1094,7 +1102,7 @@ namespace UE::NNERuntimeIREE
 					iree_hal_device_release(Device);
 				}
 
-				static TSharedPtr<FDevice> Make(const FString& BasePath)
+				static TSharedPtr<FDevice> Make(const TMap<FString, TConstArrayView<uint8>>& Executables)
 				{
 					TSharedPtr<UE::NNERuntimeIREE::Private::FInstance> Instance = UE::NNERuntimeIREE::Private::FInstance::GetInstance();
 					if (!Instance.IsValid())
@@ -1105,7 +1113,7 @@ namespace UE::NNERuntimeIREE
 					iree_string_view_t Identifier = iree_make_cstring_view("unreal");
 
 					iree_hal_device_t* Device = nullptr;
-					iree_status_t Status = IREE::HAL::RDG::DeviceCreate(Identifier, iree_allocator_system(), BasePath, &Device);
+					iree_status_t Status = IREE::HAL::RDG::DeviceCreate(Identifier, iree_allocator_system(), Executables, &Device);
 					if (!iree_status_is_ok(Status))
 					{
 						UE_LOG(LogNNERuntimeIREE, Error, TEXT("Could not create IREE RDG device!"));
@@ -1561,18 +1569,30 @@ namespace UE::NNERuntimeIREE
 
 		}
 
-		TSharedPtr<FModel> FModel::Make(const FString& InDirPath, const FString& InVmfbFileName, const UNNERuntimeIREEModuleMetaData& InModuleMetaData)
+		TSharedPtr<FModel> FModel::Make(const FString& InDirPath, TConstArrayView64<uint8> VmfbData, const UNNERuntimeIREEModuleMetaData& InModuleMetaData, const TMap<FString, TConstArrayView<uint8>>& Executables)
 		{
-			check(!InVmfbFileName.IsEmpty());
+			check(!InDirPath.IsEmpty());
+			check(!VmfbData.IsEmpty());
 			check(!InModuleMetaData.FunctionMetaData.IsEmpty());
 
-			TSharedPtr<Private::FDevice> Device = Private::FDevice::Make(InDirPath);
+			TSharedPtr<Private::FDevice> Device = Private::FDevice::Make(Executables);
 			if (!Device.IsValid())
 			{
 				return {};
 			}
 
-			TSharedPtr<UE::NNERuntimeIREE::Private::FModule> Module = UE::NNERuntimeIREE::Private::FModule::Make(InDirPath, InVmfbFileName, InModuleMetaData);
+			const int64 DataSize = VmfbData.Num();
+
+			void* Data = FMemory::Malloc(DataSize, IREE_HAL_HEAP_BUFFER_ALIGNMENT);
+			FMemory::Memcpy(Data, VmfbData.GetData(), DataSize);
+
+			TSharedPtr<UE::NNE::FSharedModelData> ModelData = MakeShared<UE::NNE::FSharedModelData>(FSharedBuffer::TakeOwnership(Data, DataSize, FMemory::Free), IREE_HAL_HEAP_BUFFER_ALIGNMENT);
+			if (!ModelData.IsValid())
+			{
+				return {};
+			}
+
+			TSharedPtr<UE::NNERuntimeIREE::Private::FModule> Module = UE::NNERuntimeIREE::Private::FModule::Make(InDirPath, ModelData.ToSharedRef(), InModuleMetaData);
 			if (!Module.IsValid())
 			{
 				return {};
diff --git a/Engine/Plugins/Experimental/NNERuntimeIREE/Source/NNERuntimeIREE/Private/NNERuntimeIREEModel.h b/Engine/Plugins/Experimental/NNERuntimeIREE/Source/NNERuntimeIREE/Private/NNERuntimeIREEModel.h
--- a/Engine/Plugins/Experimental/NNERuntimeIREE/Source/NNERuntimeIREE/Private/NNERuntimeIREEModel.h
+++ b/Engine/Plugins/Experimental/NNERuntimeIREE/Source/NNERuntimeIREE/Private/NNERuntimeIREEModel.h
@@ -108,7 +108,7 @@ namespace UE::NNERuntimeIREE
 			FModel(TSharedRef<UE::NNERuntimeIREE::RDG::Private::FDevice> InDevice, TSharedRef<UE::NNERuntimeIREE::Private::FModule> InModule);
 
 		public:
-			static TSharedPtr<FModel> Make(const FString& InDirPath, const FString& InVmfbFileName, const UNNERuntimeIREEModuleMetaData& InModuleMetaData);
+			static TSharedPtr<FModel> Make(const FString& InDirPath, TConstArrayView64<uint8> VmfbData, const UNNERuntimeIREEModuleMetaData& InModuleMetaData, const TMap<FString, TConstArrayView<uint8>>& Executables);
 
 		public:
 			//~ Begin IModelRDG Interface
diff --git a/Engine/Plugins/Experimental/NNERuntimeIREE/Source/NNERuntimeIREE/Private/NNERuntimeIREEModelData.cpp b/Engine/Plugins/Experimental/NNERuntimeIREE/Source/NNERuntimeIREE/Private/NNERuntimeIREEModelData.cpp
--- a/Engine/Plugins/Experimental/NNERuntimeIREE/Source/NNERuntimeIREE/Private/NNERuntimeIREEModelData.cpp
+++ b/Engine/Plugins/Experimental/NNERuntimeIREE/Source/NNERuntimeIREE/Private/NNERuntimeIREEModelData.cpp
@@ -7,7 +7,7 @@
 #include "NNERuntimeIREELog.h"
 #include "Serialization/CustomVersion.h"
 
-namespace UE::NNERuntimeIREE::ModelData::Private
+namespace UE::NNERuntimeIREE::CPU::ModelData::Private
 {
 
 enum Version : uint32
@@ -18,14 +18,14 @@ enum Version : uint32
 	Latest = VersionPlusOne - 1
 };
 const FGuid GUID(0x6dcb835d, 0x9ac64a1d, 0x8165d871, 0x6122dab7);
-FCustomVersionRegistration Version(GUID, Version::Latest, TEXT("NNERuntimeIREEModelDataVersion"));// Always save with the latest version
+FCustomVersionRegistration Version(GUID, Version::Latest, TEXT("NNERuntimeIREEModelDataVersionCPU"));// Always save with the latest version
 
-} // UE::NNERuntimeIREE::ModelData::Private
+} // UE::NNERuntimeIREE::CPU::ModelData::Private
 
-void UNNERuntimeIREEModelData::Serialize(FArchive& Ar)
+void UNNERuntimeIREEModelDataCPU::Serialize(FArchive& Ar)
 {
 	// Store the asset version (no effect in load)
-	Ar.UsingCustomVersion(UE::NNERuntimeIREE::ModelData::Private::GUID);
+	Ar.UsingCustomVersion(UE::NNERuntimeIREE::CPU::ModelData::Private::GUID);
 
 	if (Ar.IsSaving() || Ar.IsCountingMemory()) 
 	{
@@ -39,9 +39,9 @@ void UNNERuntimeIREEModelData::Serialize(FArchive& Ar)
 	{
 		int32 NumArchitectures = 0;
 
-		switch (Ar.CustomVer(UE::NNERuntimeIREE::ModelData::Private::GUID))
+		switch (Ar.CustomVer(UE::NNERuntimeIREE::CPU::ModelData::Private::GUID))
 		{
-		case UE::NNERuntimeIREE::ModelData::Private::Version::V0:
+		case UE::NNERuntimeIREE::CPU::ModelData::Private::Version::V0:
 			Ar << GUID;
 			Ar << Version;
 			Ar << FileId;
@@ -49,16 +49,16 @@ void UNNERuntimeIREEModelData::Serialize(FArchive& Ar)
 			Ar << CompilerResult;
 			break;
 		default:
-			UE_LOG(LogNNERuntimeIREE, Error, TEXT("UNNERuntimeIREEModelData: Unknown asset version %d: Deserialisation failed, please reimport the original model."), Ar.CustomVer(UE::NNERuntimeIREE::ModelData::Private::GUID));
+			UE_LOG(LogNNERuntimeIREE, Error, TEXT("UNNERuntimeIREEModelDataCPU: Unknown asset version %d: Deserialisation failed, please reimport the original model."), Ar.CustomVer(UE::NNERuntimeIREE::CPU::ModelData::Private::GUID));
 			break;
 		}
 	}
 }
 
-bool UNNERuntimeIREEModelData::IsSameGuidAndVersion(TConstArrayView64<uint8> Data, FGuid Guid, int32 Version)
+bool UNNERuntimeIREEModelDataCPU::IsSameGuidAndVersion(TConstArrayView64<uint8> Data, FGuid Guid, int32 Version)
 {
-	int32 GuidSize = sizeof(UNNERuntimeIREEModelData::GUID);
-	int32 VersionSize = sizeof(UNNERuntimeIREEModelData::Version);
+	int32 GuidSize = sizeof(UNNERuntimeIREEModelDataCPU::GUID);
+	int32 VersionSize = sizeof(UNNERuntimeIREEModelDataCPU::Version);
 	if (Data.Num() <= GuidSize + VersionSize)
 	{
 		return false;
@@ -68,4 +68,62 @@ bool UNNERuntimeIREEModelData::IsSameGuidAndVersion(TConstArrayView64<uint8> Dat
 	bResult &= FGenericPlatformMemory::Memcmp(&(Data[GuidSize]), &Version, VersionSize) == 0;
 
 	return bResult;
+}
+
+namespace UE::NNERuntimeIREE::RDG::ModelData::Private
+{
+
+enum Version : uint32
+{
+	V0 = 0, // Initial
+	// New versions can be added above this line
+	VersionPlusOne,
+	Latest = VersionPlusOne - 1
+};
+const FGuid GUID(0x29663d20, 0x1dbd43c5, 0xa1ecc03b, 0x6005c845);
+FCustomVersionRegistration Version(GUID, Version::Latest, TEXT("NNERuntimeIREEModelDataVersionRDG"));// Always save with the latest version
+
+} // UE::NNERuntimeIREE::RDG::ModelData::Private
+
+void UNNERuntimeIREEModelDataRDG::Serialize(FArchive& Ar)
+{
+	// Store the asset version (no effect in load)
+	Ar.UsingCustomVersion(UE::NNERuntimeIREE::RDG::ModelData::Private::GUID);
+
+	if (Ar.IsSaving() || Ar.IsCountingMemory()) 
+	{
+		FNNERuntimeIREEModelDataHeaderRDG::StaticStruct()->SerializeBin(Ar, &Header);
+		Ar << ModuleMetaData;
+		Ar << CompilerResult;
+	}
+	else
+	{
+		switch (Ar.CustomVer(UE::NNERuntimeIREE::RDG::ModelData::Private::GUID))
+		{
+		case UE::NNERuntimeIREE::RDG::ModelData::Private::Version::V0:
+		FNNERuntimeIREEModelDataHeaderRDG::StaticStruct()->SerializeBin(Ar, &Header);
+			Ar << ModuleMetaData;
+			Ar << CompilerResult;
+			break;
+		default:
+			UE_LOG(LogNNERuntimeIREE, Error, TEXT("UNNERuntimeIREEModelDataRDG: Unknown asset version %d: Deserialisation failed, please reimport the original model."), Ar.CustomVer(UE::NNERuntimeIREE::RDG::ModelData::Private::GUID));
+			break;
+		}
+	}
+}
+
+FNNERuntimeIREEModelDataHeaderRDG UNNERuntimeIREEModelDataRDG::ReadHeader(FArchive& Ar)
+{
+	FNNERuntimeIREEModelDataHeaderRDG Header{};
+	switch (Ar.CustomVer(UE::NNERuntimeIREE::RDG::ModelData::Private::GUID))
+	{
+	case UE::NNERuntimeIREE::RDG::ModelData::Private::Version::V0:
+	FNNERuntimeIREEModelDataHeaderRDG::StaticStruct()->SerializeBin(Ar, &Header);
+		break;
+	default:
+		UE_LOG(LogNNERuntimeIREE, Error, TEXT("UNNERuntimeIREEModelDataRDG: Unknown asset version %d: Deserialisation failed, please reimport the original model."), Ar.CustomVer(UE::NNERuntimeIREE::RDG::ModelData::Private::GUID));
+		break;
+	}
+
+	return Header;
 }
\ No newline at end of file
diff --git a/Engine/Plugins/Experimental/NNERuntimeIREE/Source/NNERuntimeIREE/Private/NNERuntimeIREEModelData.h b/Engine/Plugins/Experimental/NNERuntimeIREE/Source/NNERuntimeIREE/Private/NNERuntimeIREEModelData.h
--- a/Engine/Plugins/Experimental/NNERuntimeIREE/Source/NNERuntimeIREE/Private/NNERuntimeIREEModelData.h
+++ b/Engine/Plugins/Experimental/NNERuntimeIREE/Source/NNERuntimeIREE/Private/NNERuntimeIREEModelData.h
@@ -13,7 +13,7 @@
  * IREE model data class.
  */
 UCLASS()
-class UNNERuntimeIREEModelData : public UObject
+class UNNERuntimeIREEModelDataCPU : public UObject
 {
 	GENERATED_BODY()
 
@@ -52,3 +52,67 @@ public:
 	 */
 	TArray64<uint8> CompilerResult;
 };
+
+/**
+ * IREE model data header RDG struct.
+ */
+USTRUCT()
+struct FNNERuntimeIREEModelDataHeaderRDG
+{
+	GENERATED_BODY()
+
+	/**
+	 * A Guid that uniquely identifies this IREE model data.
+	 */
+	UPROPERTY()
+	FGuid GUID;
+
+	/**
+	 * Current version of this IREE model data.
+	 */
+	UPROPERTY()
+	int32 Version;
+
+	/**
+	 * A Guid that uniquely identifies the model.
+	 */
+	UPROPERTY()
+	FGuid FileId;
+
+	/**
+	 * List of shader platforms supported by this IREE model data.
+	 */
+	UPROPERTY()
+	TArray<FString> ShaderPlatforms;
+};
+
+/**
+ * IREE model data RDG class.
+ */
+UCLASS()
+class UNNERuntimeIREEModelDataRDG : public UObject
+{
+	GENERATED_BODY()
+
+public:
+	// UObject interface
+	virtual void Serialize(FArchive& Ar) override;
+	// End of UObject interface
+
+	static FNNERuntimeIREEModelDataHeaderRDG ReadHeader(FArchive& Ar);
+
+	/**
+	 * Model data header.
+	 */
+	FNNERuntimeIREEModelDataHeaderRDG Header;
+
+	/**
+	 * Serialized module meta data.
+	 */
+	TArray64<uint8> ModuleMetaData;
+
+	/**
+	 * Serialized compiler output.
+	 */
+	TArray64<uint8> CompilerResult;
+};
\ No newline at end of file
