[NNE] NNERuntimeIREERdg add define to toggle compilation.

#tests

#rnx

diff --git a/Engine/Plugins/Experimental/NNERuntimeIREE/Source/NNERuntimeIREE/Private/NNERuntimeIREEModel.cpp b/Engine/Plugins/Experimental/NNERuntimeIREE/Source/NNERuntimeIREE/Private/NNERuntimeIREEModel.cpp
--- a/Engine/Plugins/Experimental/NNERuntimeIREE/Source/NNERuntimeIREE/Private/NNERuntimeIREEModel.cpp
+++ b/Engine/Plugins/Experimental/NNERuntimeIREE/Source/NNERuntimeIREE/Private/NNERuntimeIREEModel.cpp
@@ -1085,6 +1085,7 @@ namespace UE::NNERuntimeIREE
 		}
 	} // CPU
 
+#ifdef WITH_NNE_RUNTIME_IREE_RDG
 	namespace RDG
 	{
 		namespace Private
@@ -1607,7 +1608,7 @@ namespace UE::NNERuntimeIREE
 		}
 
 	} // namespace RDG
-
+#endif // WITH_NNE_RUNTIME_IREE_RDG
 } // UE::NNERuntimeIREE
 
 #endif // WITH_NNE_RUNTIME_IREE
\ No newline at end of file
diff --git a/Engine/Plugins/Experimental/NNERuntimeIREE/Source/NNERuntimeIREE/Private/NNERuntimeIREEModel.h b/Engine/Plugins/Experimental/NNERuntimeIREE/Source/NNERuntimeIREE/Private/NNERuntimeIREEModel.h
--- a/Engine/Plugins/Experimental/NNERuntimeIREE/Source/NNERuntimeIREE/Private/NNERuntimeIREEModel.h
+++ b/Engine/Plugins/Experimental/NNERuntimeIREE/Source/NNERuntimeIREE/Private/NNERuntimeIREEModel.h
@@ -5,7 +5,9 @@
 #ifdef WITH_NNE_RUNTIME_IREE
 
 #include "NNERuntimeCPU.h"
+#ifdef WITH_NNE_RUNTIME_IREE_RDG
 #include "NNERuntimeRDG.h"
+#endif // WITH_NNE_RUNTIME_IREE_RDG
 #include "NNERuntimeIREEMetaData.h"
 
 LLM_DECLARE_TAG(NNERuntimeIREE_Cpu);
@@ -69,6 +71,7 @@ namespace UE::NNERuntimeIREE
 		};
 	} // CPU
 
+#ifdef WITH_NNE_RUNTIME_IREE_RDG
 	namespace RDG
 	{
 		namespace Private 
@@ -120,7 +123,7 @@ namespace UE::NNERuntimeIREE
 			TSharedRef<UE::NNERuntimeIREE::Private::FModule> Module;
 		};
 	} // namespace RDG
-
+#endif // WITH_NNE_RUNTIME_IREE_RDG
 } // UE::NNERuntimeIREE
 
 #endif // WITH_NNE_RUNTIME_IREE
\ No newline at end of file
diff --git a/Engine/Plugins/Experimental/NNERuntimeIREE/Source/NNERuntimeIREEShader/Internal/NNERuntimeIREEShader.h b/Engine/Plugins/Experimental/NNERuntimeIREE/Source/NNERuntimeIREEShader/Internal/NNERuntimeIREEShader.h
--- a/Engine/Plugins/Experimental/NNERuntimeIREE/Source/NNERuntimeIREEShader/Internal/NNERuntimeIREEShader.h
+++ b/Engine/Plugins/Experimental/NNERuntimeIREE/Source/NNERuntimeIREEShader/Internal/NNERuntimeIREEShader.h
@@ -2,6 +2,8 @@
 
 #pragma once
 
+#ifdef WITH_NNE_RUNTIME_IREE_SHADER
+
 #include "Shader.h"
 #include "ShaderParameters.h"
 #include "NNERuntimeIREEShaderType.h"
@@ -18,4 +20,6 @@ public:
 	MS_ALIGN(SHADER_PARAMETER_STRUCT_ALIGNMENT) struct FParameters
 	{
 	};
-};
\ No newline at end of file
+};
+
+#endif // WITH_NNE_RUNTIME_IREE_SHADER
\ No newline at end of file
diff --git a/Engine/Plugins/Experimental/NNERuntimeIREE/Source/NNERuntimeIREEShader/Internal/NNERuntimeIREEShaderCompileResult.h b/Engine/Plugins/Experimental/NNERuntimeIREE/Source/NNERuntimeIREEShader/Internal/NNERuntimeIREEShaderCompileResult.h
--- a/Engine/Plugins/Experimental/NNERuntimeIREE/Source/NNERuntimeIREEShader/Internal/NNERuntimeIREEShaderCompileResult.h
+++ b/Engine/Plugins/Experimental/NNERuntimeIREE/Source/NNERuntimeIREEShader/Internal/NNERuntimeIREEShaderCompileResult.h
@@ -2,6 +2,8 @@
 
 #pragma once
 
+#ifdef WITH_NNE_RUNTIME_IREE_SHADER
+
 #include "CoreMinimal.h"
 
 /** Struct containing a message returned by shader compilation. */
@@ -35,3 +37,5 @@ struct NNERUNTIMEIREESHADER_API FNNERuntimeIREEShaderCompileResults
 {
 	TArray<FNNERuntimeIREEShaderCompileMessage> Messages;
 };
+
+#endif // WITH_NNE_RUNTIME_IREE_SHADER
\ No newline at end of file
diff --git a/Engine/Plugins/Experimental/NNERuntimeIREE/Source/NNERuntimeIREEShader/Internal/NNERuntimeIREEShaderFillBufferCS.h b/Engine/Plugins/Experimental/NNERuntimeIREE/Source/NNERuntimeIREEShader/Internal/NNERuntimeIREEShaderFillBufferCS.h
--- a/Engine/Plugins/Experimental/NNERuntimeIREE/Source/NNERuntimeIREEShader/Internal/NNERuntimeIREEShaderFillBufferCS.h
+++ b/Engine/Plugins/Experimental/NNERuntimeIREE/Source/NNERuntimeIREEShader/Internal/NNERuntimeIREEShaderFillBufferCS.h
@@ -2,6 +2,8 @@
 
 #pragma once
 
+#ifdef WITH_NNE_RUNTIME_IREE_SHADER
+
 #include "GlobalShader.h"
 #include "RenderGraphUtils.h"
 #include "ShaderParameterUtils.h"
@@ -32,3 +34,5 @@ public:
 };
 
 } // UE::NNERuntimeIREEShader::Internal
+
+#endif // WITH_NNE_RUNTIME_IREE_SHADER
\ No newline at end of file
diff --git a/Engine/Plugins/Experimental/NNERuntimeIREE/Source/NNERuntimeIREEShader/Internal/NNERuntimeIREEShaderMetadataAllocations.h b/Engine/Plugins/Experimental/NNERuntimeIREE/Source/NNERuntimeIREEShader/Internal/NNERuntimeIREEShaderMetadataAllocations.h
--- a/Engine/Plugins/Experimental/NNERuntimeIREE/Source/NNERuntimeIREEShader/Internal/NNERuntimeIREEShaderMetadataAllocations.h
+++ b/Engine/Plugins/Experimental/NNERuntimeIREE/Source/NNERuntimeIREEShader/Internal/NNERuntimeIREEShaderMetadataAllocations.h
@@ -2,6 +2,8 @@
 
 #pragma once
 
+#ifdef WITH_NNE_RUNTIME_IREE_SHADER
+
 #include "Containers/Array.h"
 #include "Containers/UnrealString.h"
 #include "ShaderParameterMetadata.h"
@@ -18,4 +20,6 @@ struct NNERUNTIMEIREESHADER_API FNNERuntimeIREEShaderParametersMetadataAllocatio
 	FNNERuntimeIREEShaderParametersMetadataAllocations(FNNERuntimeIREEShaderParametersMetadataAllocations& Other) = delete;
 
 	~FNNERuntimeIREEShaderParametersMetadataAllocations();
-};
\ No newline at end of file
+};
+
+#endif // WITH_NNE_RUNTIME_IREE_SHADER
\ No newline at end of file
diff --git a/Engine/Plugins/Experimental/NNERuntimeIREE/Source/NNERuntimeIREEShader/Internal/NNERuntimeIREEShaderShared.h b/Engine/Plugins/Experimental/NNERuntimeIREE/Source/NNERuntimeIREEShader/Internal/NNERuntimeIREEShaderShared.h
--- a/Engine/Plugins/Experimental/NNERuntimeIREE/Source/NNERuntimeIREEShader/Internal/NNERuntimeIREEShaderShared.h
+++ b/Engine/Plugins/Experimental/NNERuntimeIREE/Source/NNERuntimeIREEShader/Internal/NNERuntimeIREEShaderShared.h
@@ -2,6 +2,8 @@
 
 #pragma once
 
+#ifdef WITH_NNE_RUNTIME_IREE_SHADER
+
 #include "CoreMinimal.h"
 #include "NNERuntimeIREEShaderCompileResult.h"
 #include "Engine/EngineTypes.h"
@@ -437,4 +439,6 @@ private:
 
 	uint32 bLoadedCookedShaderMapId : 1;
 	FNNERuntimeIREEShaderMapId CookedShaderMapId;
-};
\ No newline at end of file
+};
+
+#endif // WITH_NNE_RUNTIME_IREE_SHADER
\ No newline at end of file
diff --git a/Engine/Plugins/Experimental/NNERuntimeIREE/Source/NNERuntimeIREEShader/Internal/NNERuntimeIREEShaderType.h b/Engine/Plugins/Experimental/NNERuntimeIREE/Source/NNERuntimeIREEShader/Internal/NNERuntimeIREEShaderType.h
--- a/Engine/Plugins/Experimental/NNERuntimeIREE/Source/NNERuntimeIREEShader/Internal/NNERuntimeIREEShaderType.h
+++ b/Engine/Plugins/Experimental/NNERuntimeIREE/Source/NNERuntimeIREEShader/Internal/NNERuntimeIREEShaderType.h
@@ -2,6 +2,8 @@
 
 #pragma once
 
+#ifdef WITH_NNE_RUNTIME_IREE_SHADER
+
 #include "CoreMinimal.h"
 #include "Engine/EngineTypes.h"
 #include "GlobalShader.h"
@@ -165,4 +167,6 @@ protected:
 		}
 	}
 #endif // WITH_EDITOR
-};
\ No newline at end of file
+};
+
+#endif // WITH_NNE_RUNTIME_IREE_SHADER
\ No newline at end of file
diff --git a/Engine/Plugins/Experimental/NNERuntimeIREE/Source/NNERuntimeIREEShader/NNERuntimeIREEShader.Build.cs b/Engine/Plugins/Experimental/NNERuntimeIREE/Source/NNERuntimeIREEShader/NNERuntimeIREEShader.Build.cs
--- a/Engine/Plugins/Experimental/NNERuntimeIREE/Source/NNERuntimeIREEShader/NNERuntimeIREEShader.Build.cs
+++ b/Engine/Plugins/Experimental/NNERuntimeIREE/Source/NNERuntimeIREEShader/NNERuntimeIREEShader.Build.cs
@@ -20,5 +20,18 @@ public class NNERuntimeIREEShader : ModuleRules
 				"RHI",
 			}
 		);
+
+		if (Target.Platform == UnrealTargetPlatform.Win64)
+		{
+			PublicDefinitions.Add("WITH_NNE_RUNTIME_IREE_SHADER");
+		}
+		else if (Target.Platform == UnrealTargetPlatform.Linux)
+		{
+			PublicDefinitions.Add("WITH_NNE_RUNTIME_IREE_SHADER");
+		}
+		else if (Target.Platform == UnrealTargetPlatform.Mac)
+		{
+			PublicDefinitions.Add("WITH_NNE_RUNTIME_IREE_SHADER");
+		}
 	}
 }
diff --git a/Engine/Plugins/Experimental/NNERuntimeIREE/Source/NNERuntimeIREEShader/Private/NNERuntimeIREEShader.cpp b/Engine/Plugins/Experimental/NNERuntimeIREE/Source/NNERuntimeIREEShader/Private/NNERuntimeIREEShader.cpp
--- a/Engine/Plugins/Experimental/NNERuntimeIREE/Source/NNERuntimeIREEShader/Private/NNERuntimeIREEShader.cpp
+++ b/Engine/Plugins/Experimental/NNERuntimeIREE/Source/NNERuntimeIREEShader/Private/NNERuntimeIREEShader.cpp
@@ -2,6 +2,8 @@
 
 #include "NNERuntimeIREEShader.h"
 
+#ifdef WITH_NNE_RUNTIME_IREE_SHADER
+
 #include "NNERuntimeIREEShaderType.h"
 
 IMPLEMENT_SHADER_TYPE(, FNNERuntimeIREEShader, TEXT("/Plugin/NNERuntimeIREEShader/NNERuntimeIREEShader.usf"), TEXT("__"), SF_Compute)
@@ -12,4 +14,6 @@ FNNERuntimeIREEShader::FNNERuntimeIREEShader(const FNNERuntimeIREEShaderType::Co
 	const FShaderParametersMetadata& ShaderParametersMetadata = static_cast<const FNNERuntimeIREEShaderType::FParameters*>(Initializer.Parameters)->ShaderParamMetadata;
 
 	Bindings.BindForLegacyShaderParameters(this, Initializer.PermutationId, Initializer.ParameterMap, ShaderParametersMetadata, true);
-}
\ No newline at end of file
+}
+
+#endif // WITH_NNE_RUNTIME_IREE_SHADER
\ No newline at end of file
diff --git a/Engine/Plugins/Experimental/NNERuntimeIREE/Source/NNERuntimeIREEShader/Private/NNERuntimeIREEShaderCompilationManager.cpp b/Engine/Plugins/Experimental/NNERuntimeIREE/Source/NNERuntimeIREEShader/Private/NNERuntimeIREEShaderCompilationManager.cpp
--- a/Engine/Plugins/Experimental/NNERuntimeIREE/Source/NNERuntimeIREEShader/Private/NNERuntimeIREEShaderCompilationManager.cpp
+++ b/Engine/Plugins/Experimental/NNERuntimeIREE/Source/NNERuntimeIREEShader/Private/NNERuntimeIREEShaderCompilationManager.cpp
@@ -2,6 +2,8 @@
 
 #include "NNERuntimeIREEShaderCompilationManager.h"
 
+#ifdef WITH_NNE_RUNTIME_IREE_SHADER
+
 #include "DataDrivenShaderPlatformInfo.h"
 #include "HAL/FileManager.h"
 #include "Misc/Paths.h"
@@ -412,3 +414,4 @@ void FNNERuntimeIREEShaderCompilationManager::FinishCompilation(const TCHAR* InK
 }
 
 #endif // WITH_EDITOR
+#endif // WITH_NNE_RUNTIME_IREE_SHADER
\ No newline at end of file
diff --git a/Engine/Plugins/Experimental/NNERuntimeIREE/Source/NNERuntimeIREEShader/Private/NNERuntimeIREEShaderCompilationManager.h b/Engine/Plugins/Experimental/NNERuntimeIREE/Source/NNERuntimeIREEShader/Private/NNERuntimeIREEShaderCompilationManager.h
--- a/Engine/Plugins/Experimental/NNERuntimeIREE/Source/NNERuntimeIREEShader/Private/NNERuntimeIREEShaderCompilationManager.h
+++ b/Engine/Plugins/Experimental/NNERuntimeIREE/Source/NNERuntimeIREEShader/Private/NNERuntimeIREEShaderCompilationManager.h
@@ -2,6 +2,8 @@
 
 #pragma once
 
+#ifdef WITH_NNE_RUNTIME_IREE_SHADER
+
 #include "NNERuntimeIREEShaderShared.h"
 #include "ShaderCompiler.h"
 
@@ -58,3 +60,4 @@ private:
 extern FNNERuntimeIREEShaderCompilationManager GNNERuntimeIREEShaderCompilationManager;
 
 #endif // WITH_EDITOR
+#endif // WITH_NNE_RUNTIME_IREE_SHADER
\ No newline at end of file
diff --git a/Engine/Plugins/Experimental/NNERuntimeIREE/Source/NNERuntimeIREEShader/Private/NNERuntimeIREEShaderFillBufferCS.cpp b/Engine/Plugins/Experimental/NNERuntimeIREE/Source/NNERuntimeIREEShader/Private/NNERuntimeIREEShaderFillBufferCS.cpp
--- a/Engine/Plugins/Experimental/NNERuntimeIREE/Source/NNERuntimeIREEShader/Private/NNERuntimeIREEShaderFillBufferCS.cpp
+++ b/Engine/Plugins/Experimental/NNERuntimeIREE/Source/NNERuntimeIREEShader/Private/NNERuntimeIREEShaderFillBufferCS.cpp
@@ -2,6 +2,8 @@
 
 #include "NNERuntimeIREEShaderFillBufferCS.h"
 
+#ifdef WITH_NNE_RUNTIME_IREE_SHADER
+
 namespace UE::NNERuntimeIREEShader::Internal
 {
 
@@ -20,4 +22,6 @@ bool FFillBufferCS::ShouldCompilePermutation(const FGlobalShaderPermutationParam
 
 IMPLEMENT_GLOBAL_SHADER(FFillBufferCS, "/Plugin/NNERuntimeIREEShader/NNERuntimeIREEShaderFillBuffer.usf", "Main", SF_Compute);
 
-} // UE::NNERuntimeIREEShader::Internal
\ No newline at end of file
+} // UE::NNERuntimeIREEShader::Internal
+
+#endif // WITH_NNE_RUNTIME_IREE_SHADER
\ No newline at end of file
diff --git a/Engine/Plugins/Experimental/NNERuntimeIREE/Source/NNERuntimeIREEShader/Private/NNERuntimeIREEShaderMap.cpp b/Engine/Plugins/Experimental/NNERuntimeIREE/Source/NNERuntimeIREEShader/Private/NNERuntimeIREEShaderMap.cpp
--- a/Engine/Plugins/Experimental/NNERuntimeIREE/Source/NNERuntimeIREEShader/Private/NNERuntimeIREEShaderMap.cpp
+++ b/Engine/Plugins/Experimental/NNERuntimeIREE/Source/NNERuntimeIREEShader/Private/NNERuntimeIREEShaderMap.cpp
@@ -2,6 +2,8 @@
 
 #include "NNERuntimeIREEShader.h"
 
+#ifdef WITH_NNE_RUNTIME_IREE_SHADER
+
 #include "NNERuntimeIREEShaderShared.h"
 #include "NNERuntimeIREEShaderCompilationManager.h"
 #include "NNERuntimeIREEShaderLog.h"
@@ -523,4 +525,6 @@ bool FNNERuntimeIREEShaderMap::Serialize(FArchive& Ar)
 {
 	FShaderSerializeContext Ctx(Ar);
 	return Super::Serialize(Ctx);
-}
\ No newline at end of file
+}
+
+#endif // WITH_NNE_RUNTIME_IREE_SHADER
\ No newline at end of file
diff --git a/Engine/Plugins/Experimental/NNERuntimeIREE/Source/NNERuntimeIREEShader/Private/NNERuntimeIREEShaderMetadataAllocations.cpp b/Engine/Plugins/Experimental/NNERuntimeIREE/Source/NNERuntimeIREEShader/Private/NNERuntimeIREEShaderMetadataAllocations.cpp
--- a/Engine/Plugins/Experimental/NNERuntimeIREE/Source/NNERuntimeIREEShader/Private/NNERuntimeIREEShaderMetadataAllocations.cpp
+++ b/Engine/Plugins/Experimental/NNERuntimeIREE/Source/NNERuntimeIREEShader/Private/NNERuntimeIREEShaderMetadataAllocations.cpp
@@ -2,7 +2,11 @@
 
 #include "NNERuntimeIREEShaderMetadataAllocations.h"
 
+#ifdef WITH_NNE_RUNTIME_IREE_SHADER
+
 FNNERuntimeIREEShaderParametersMetadataAllocations::~FNNERuntimeIREEShaderParametersMetadataAllocations()
 {
 	
-}
\ No newline at end of file
+}
+
+#endif // WITH_NNE_RUNTIME_IREE_SHADER
\ No newline at end of file
diff --git a/Engine/Plugins/Experimental/NNERuntimeIREE/Source/NNERuntimeIREEShader/Private/NNERuntimeIREEShaderModule.cpp b/Engine/Plugins/Experimental/NNERuntimeIREE/Source/NNERuntimeIREEShader/Private/NNERuntimeIREEShaderModule.cpp
--- a/Engine/Plugins/Experimental/NNERuntimeIREE/Source/NNERuntimeIREEShader/Private/NNERuntimeIREEShaderModule.cpp
+++ b/Engine/Plugins/Experimental/NNERuntimeIREE/Source/NNERuntimeIREEShader/Private/NNERuntimeIREEShaderModule.cpp
@@ -12,6 +12,7 @@ DEFINE_LOG_CATEGORY(LogNNERuntimeIREEShader);
 
 void FNNERuntimeIREEShaderModule::StartupModule()
 {
+#ifdef WITH_NNE_RUNTIME_IREE_SHADER
 	const TSharedPtr<IPlugin> Plugin = IPluginManager::Get().FindPlugin(TEXT("NNERuntimeIREE"));
 	if (Plugin.IsValid())
 	{
@@ -23,6 +24,7 @@ void FNNERuntimeIREEShaderModule::StartupModule()
 	{
 		UE_LOG(LogNNERuntimeIREEShader, Warning, TEXT("Shaders directory not added. Failed to find NNERuntimeIREE plugin"));
 	}
+#endif // WITH_NNE_RUNTIME_IREE_SHADER
 }
 
 void FNNERuntimeIREEShaderModule::ShutdownModule()
diff --git a/Engine/Plugins/Experimental/NNERuntimeIREE/Source/NNERuntimeIREEShader/Private/NNERuntimeIREEShaderShared.cpp b/Engine/Plugins/Experimental/NNERuntimeIREE/Source/NNERuntimeIREEShader/Private/NNERuntimeIREEShaderShared.cpp
--- a/Engine/Plugins/Experimental/NNERuntimeIREE/Source/NNERuntimeIREEShader/Private/NNERuntimeIREEShaderShared.cpp
+++ b/Engine/Plugins/Experimental/NNERuntimeIREE/Source/NNERuntimeIREEShader/Private/NNERuntimeIREEShaderShared.cpp
@@ -2,6 +2,8 @@
 
 #include "NNERuntimeIREEShaderShared.h"
 
+#ifdef WITH_NNE_RUNTIME_IREE_SHADER
+
 #include "DataDrivenShaderPlatformInfo.h"
 #include "Interfaces/ITargetPlatform.h"
 #include "NNERuntimeIREEShader.h"
@@ -352,4 +354,6 @@ bool FNNERuntimeIREEShaderMapId::ContainsShaderType(const FShaderType* ShaderTyp
 	}
 
 	return false;
-}
\ No newline at end of file
+}
+
+#endif // WITH_NNE_RUNTIME_IREE_SHADER
\ No newline at end of file
