[NNE] NNERuntimeIREE disable compilation for Windows Arm64 until IREE suports it.

#tests compilation on Windows x64 and Arm64

#rnx

diff --git a/Engine/Plugins/Experimental/NNERuntimeIREE/Source/IREEDriverRDG/IREEDriverRDG.Build.cs b/Engine/Plugins/Experimental/NNERuntimeIREE/Source/IREEDriverRDG/IREEDriverRDG.Build.cs
--- a/Engine/Plugins/Experimental/NNERuntimeIREE/Source/IREEDriverRDG/IREEDriverRDG.Build.cs
+++ b/Engine/Plugins/Experimental/NNERuntimeIREE/Source/IREEDriverRDG/IREEDriverRDG.Build.cs
@@ -35,7 +35,7 @@ public class IREEDriverRDG : ModuleRules
 			);
 		}
 
-		if (Target.Platform == UnrealTargetPlatform.Win64)
+		if (Target.Platform == UnrealTargetPlatform.Win64 && Target.Architecture != UnrealArch.Arm64)
 		{
 			PublicDefinitions.Add("WITH_IREE_DRIVER_RDG");
 			PrivateDefinitions.Add("IREE_DRIVER_RDG_VERBOSITY=0");
diff --git a/Engine/Plugins/Experimental/NNERuntimeIREE/Source/NNERuntimeIREE/NNERuntimeIREE.Build.cs b/Engine/Plugins/Experimental/NNERuntimeIREE/Source/NNERuntimeIREE/NNERuntimeIREE.Build.cs
--- a/Engine/Plugins/Experimental/NNERuntimeIREE/Source/NNERuntimeIREE/NNERuntimeIREE.Build.cs
+++ b/Engine/Plugins/Experimental/NNERuntimeIREE/Source/NNERuntimeIREE/NNERuntimeIREE.Build.cs
@@ -37,7 +37,7 @@ public class NNERuntimeIREE : ModuleRules
 			);
 		}
 
-		if (Target.Platform == UnrealTargetPlatform.Win64)
+		if (Target.Platform == UnrealTargetPlatform.Win64 && Target.Architecture != UnrealArch.Arm64)
 		{
 			PublicDefinitions.Add("WITH_NNE_RUNTIME_IREE");
 			PublicDefinitions.Add("WITH_NNE_RUNTIME_IREE_RDG");
diff --git a/Engine/Plugins/Experimental/NNERuntimeIREE/Source/NNERuntimeIREEShader/NNERuntimeIREEShader.Build.cs b/Engine/Plugins/Experimental/NNERuntimeIREE/Source/NNERuntimeIREEShader/NNERuntimeIREEShader.Build.cs
--- a/Engine/Plugins/Experimental/NNERuntimeIREE/Source/NNERuntimeIREEShader/NNERuntimeIREEShader.Build.cs
+++ b/Engine/Plugins/Experimental/NNERuntimeIREE/Source/NNERuntimeIREEShader/NNERuntimeIREEShader.Build.cs
@@ -21,7 +21,7 @@ public class NNERuntimeIREEShader : ModuleRules
 			}
 		);
 
-		if (Target.Platform == UnrealTargetPlatform.Win64)
+		if (Target.Platform == UnrealTargetPlatform.Win64 && Target.Architecture != UnrealArch.Arm64)
 		{
 			PublicDefinitions.Add("WITH_NNE_RUNTIME_IREE_SHADER");
 		}
