[NNE] NNERuntimeIREE added IREE commands (import, compile and link) to logs.

#tests Editor, Win64, verified commands added to logs.

#rnx

diff --git a/Engine/Plugins/Experimental/NNERuntimeIREE/NNERuntimeIREE.uplugin b/Engine/Plugins/Experimental/NNERuntimeIREE/NNERuntimeIREE.uplugin
--- a/Engine/Plugins/Experimental/NNERuntimeIREE/NNERuntimeIREE.uplugin
+++ b/Engine/Plugins/Experimental/NNERuntimeIREE/NNERuntimeIREE.uplugin
@@ -15,6 +15,11 @@
 	"IsExperimentalVersion": true,
 	"Installed": false,
 	"Modules": [
+		{
+			"Name": "IREEUtils",
+			"Type": "Runtime",
+			"LoadingPhase": "Default"
+		},
 		{
 			"Name": "IREEDriverRDG",
 			"Type": "Runtime",
diff --git a/Engine/Plugins/Experimental/NNERuntimeIREE/Source/IREEDriverRDG/IREEDriverRDG.Build.cs b/Engine/Plugins/Experimental/NNERuntimeIREE/Source/IREEDriverRDG/IREEDriverRDG.Build.cs
--- a/Engine/Plugins/Experimental/NNERuntimeIREE/Source/IREEDriverRDG/IREEDriverRDG.Build.cs
+++ b/Engine/Plugins/Experimental/NNERuntimeIREE/Source/IREEDriverRDG/IREEDriverRDG.Build.cs
@@ -20,6 +20,7 @@ public class IREEDriverRDG : ModuleRules
 				"Projects",
 				"RenderCore",
 				"RHI",
+				"IREEUtils",
 				"NNERuntimeIREEShader"
 			}
 		);
diff --git a/Engine/Plugins/Experimental/NNERuntimeIREE/Source/IREEDriverRDG/Private/IREECompilerRDG.cpp b/Engine/Plugins/Experimental/NNERuntimeIREE/Source/IREEDriverRDG/Private/IREECompilerRDG.cpp
--- a/Engine/Plugins/Experimental/NNERuntimeIREE/Source/IREEDriverRDG/Private/IREECompilerRDG.cpp
+++ b/Engine/Plugins/Experimental/NNERuntimeIREE/Source/IREEDriverRDG/Private/IREECompilerRDG.cpp
@@ -19,9 +19,9 @@
 #include "Interfaces/ITargetPlatform.h"
 #include "IREEDriverRDGLog.h"
 #include "IREEDriverRDGShaderParametersMetadata.h"
+#include "IREEUtils.h"
 #include "Kismet/GameplayStatics.h"
 #include "Misc/FileHelper.h"
-#include "Misc/MonitoredProcess.h"
 #include "Misc/Paths.h"
 #include "NNE.h"
 #include "Serialization/ArchiveSavePackageDataBuffer.h"
@@ -35,86 +35,6 @@ namespace UE::IREE::Compiler
 	{
 		namespace Private
 		{
-			FString GetOnnxFilePath(const FString& InModelName, const FString& InOutputDir)
-			{
-				return FPaths::Combine(InOutputDir, InModelName) + ".onnx";
-			}
-
-			FString GetMlirFilePath(const FString& InModelName, const FString& InOutputDir)
-			{
-				return FPaths::Combine(InOutputDir, InModelName) + ".mlir";
-			}
-
-			bool ResolveEnvironmentVariables(FString& String)
-			{
-				const FString StartString = "$ENV{";
-				const FString EndString = "}";
-				FString ResultString = String;
-				int32 StartIndex = ResultString.Find(StartString, ESearchCase::CaseSensitive);
-				while (StartIndex != INDEX_NONE)
-				{
-					StartIndex += StartString.Len();
-					int32 EndIndex = ResultString.Find(EndString, ESearchCase::CaseSensitive, ESearchDir::FromStart, StartIndex);
-					if (EndIndex > StartIndex)
-					{
-						FString EnvironmentVariableName = ResultString.Mid(StartIndex, EndIndex - StartIndex);
-						FString EnvironmentVariableValue = FPlatformMisc::GetEnvironmentVariable(*EnvironmentVariableName);
-						if (EnvironmentVariableValue.IsEmpty())
-						{
-							return false;
-						}
-						else
-						{
-							ResultString.ReplaceInline(*(StartString + EnvironmentVariableName + EndString), *EnvironmentVariableValue, ESearchCase::CaseSensitive);
-						}
-					}
-					else
-					{
-						return false;
-					}
-					StartIndex = ResultString.Find(StartString, ESearchCase::CaseSensitive);
-				}
-				String = ResultString;
-				return true;
-			}
-
-			void RunCommand(const FString& Command, const FString& Arguments, const FString& LogFilePath = FString())
-			{
-				int32 ReturnCode = 0;
-				bool IsCanceled = false;
-
-				FMonitoredProcess Process(Command, Arguments, true);
-				Process.OnCompleted().BindLambda([&ReturnCode] (int32 _ReturnCode) { ReturnCode = _ReturnCode; });
-				Process.OnCanceled().BindLambda([&IsCanceled] (){ IsCanceled = true; });
-
-				if (!Process.Launch())
-				{
-					UE_LOG(LogIREEDriverRDG, Warning, TEXT("Failed to launch subprocess!"));
-					return;
-				}
-
-				while (Process.Update())
-				{
-					// Poll until process has finished
-				}
-
-				if (IsCanceled)
-				{
-					UE_LOG(LogIREEDriverRDG, Warning, TEXT("Execution of subprocess was canceled!"));
-				}
-				else if (ReturnCode)
-				{
-					UE_LOG(LogIREEDriverRDG, Warning, TEXT("Subprocess exited with non-zero code %d"), ReturnCode);
-				}
-
-				if (!LogFilePath.IsEmpty())
-				{
-					FFileHelper::SaveStringToFile(Process.GetFullOutputWithoutDelegate(), *LogFilePath);
-
-					UE_LOG(LogIREEDriverRDG, Log, TEXT("Saved subprocess output to: %s"), *LogFilePath);
-				}
-			}
-
 			bool CompileAndSerializeShaderFromHLSLSource(const FString& HlslFilepath, EShaderPlatform ShaderPlatform, const ITargetPlatform* TargetPlatform, const FString& Outdir, bool bDumpHLSL, FIREECompilerRDGExecutableData& ExecutableData)
 			{
 				UE_LOG(LogIREEDriverRDG, Log, TEXT("Process %s"), *HlslFilepath);
@@ -275,7 +195,7 @@ namespace UE::IREE::Compiler
 							FString TmpImporterCommand;
 							for (int32 i = 0; i < BuildConfig.ImporterCommand.Num(); i++)
 							{
-								if (ResolveEnvironmentVariables(BuildConfig.ImporterCommand[i]))
+								if (IREEUtils::ResolveEnvironmentVariables(BuildConfig.ImporterCommand[i]))
 								{
 									BuildConfig.ImporterCommand[i].ReplaceInline(*FString("${PLUGIN_DIR}"), *PluginDir);
 									BuildConfig.ImporterCommand[i].ReplaceInline(*FString("${PROJECT_DIR}"), *FPaths::ProjectDir());
@@ -299,7 +219,7 @@ namespace UE::IREE::Compiler
 							FString TmpCompilerCommand;
 							for (int32 i = 0; i < BuildConfig.CompilerCommand.Num(); i++)
 							{
-								if (ResolveEnvironmentVariables(BuildConfig.CompilerCommand[i]))
+								if (IREEUtils::ResolveEnvironmentVariables(BuildConfig.CompilerCommand[i]))
 								{
 									BuildConfig.CompilerCommand[i].ReplaceInline(*FString("${PLUGIN_DIR}"), *PluginDir);
 									BuildConfig.CompilerCommand[i].ReplaceInline(*FString("${PROJECT_DIR}"), *FPaths::ProjectDir());
@@ -346,58 +266,7 @@ namespace UE::IREE::Compiler
 
 		bool FCompiler::ImportOnnx(TConstArrayView<uint8> InFileData, const FString& InModelName, const FString& InOutputDir, TArray64<uint8>& OutMlirData)
 		{
-			SCOPED_NAMED_EVENT_TEXT("FCompiler::ImportOnnx", FColor::Magenta);
-
-			using namespace Private;
-
-			IPlatformFile& PlatformFile = FPlatformFileManager::Get().GetPlatformFile();
-
-			FString InputFilePath = GetOnnxFilePath(InModelName, InOutputDir);
-			if (!PlatformFile.FileExists(*InputFilePath))
-			{
-				SCOPED_NAMED_EVENT_TEXT("InputFile", FColor::Magenta);
-
-				if(!FFileHelper::SaveArrayToFile(InFileData, *InputFilePath))
-				{
-					UE_LOG(LogIREEDriverRDG, Warning, TEXT("IREECompilerRDG failed to save ONNX model \"%s\""), *InputFilePath);
-					return false;
-				}
-			}
-
-			FString OutputFilePath = GetMlirFilePath(InModelName, InOutputDir);
-			FString IntermediateFilePathNoExt = FPaths::Combine(InOutputDir, InModelName);
-
-			if (!ResolveEnvironmentVariables(ImporterArguments))
-			{
-				UE_LOG(LogIREEDriverRDG, Warning, TEXT("IREECompilerRDG could not replace environment variables in %s"), *ImporterArguments);
-				return false;
-			}
-			ImporterArguments.ReplaceInline(*FString("${INPUT_PATH}"), *(FString("\"") + InputFilePath + "\""));
-			ImporterArguments.ReplaceInline(*FString("${OUTPUT_PATH}"), *(FString("\"") + OutputFilePath + "\""));
-
-			{
-				SCOPED_NAMED_EVENT_TEXT("Import", FColor::Magenta);
-
-				RunCommand(ImporterCommand, ImporterArguments, IntermediateFilePathNoExt + "_import-log.txt");
-			}
-
-			if (!PlatformFile.FileExists(*OutputFilePath))
-			{
-				UE_LOG(LogIREEDriverRDG, Warning, TEXT("IREECompilerRDG failed to import the model \"%s\" using the command:"), *InputFilePath);
-				UE_LOG(LogIREEDriverRDG, Warning, TEXT("\"%s\" %s"), *ImporterCommand, *ImporterArguments);
-				return false;
-			}
-
-			{
-				SCOPED_NAMED_EVENT_TEXT("Load", FColor::Magenta);
-				if(!FFileHelper::LoadFileToArray(OutMlirData, *OutputFilePath))
-				{
-					UE_LOG(LogIREEDriverRDG, Warning, TEXT("IREECompilerRDG failed to load imported model \"%s\""), *OutputFilePath);
-					return false;
-				}
-			}
-
-			return true;
+			return IREEUtils::ImportOnnx(ImporterCommand, ImporterArguments, InFileData, InModelName, InOutputDir, OutMlirData);
 		}
 
 		bool FCompiler::CompileMlir(TConstArrayView<uint8> InFileData, const FString& InModelName, const FString& InOutputDir, TConstArrayView<EShaderPlatform> ShaderPlatforms, FIREECompilerRDGResult& OutCompilerResult)
@@ -436,7 +305,7 @@ namespace UE::IREE::Compiler
 				FString VmfbFilePath = IntermediateFilePathNoExt + ".vmfb";
 
 				FString CompilerArguments = Target->CompilerArguments;
-				if (!ResolveEnvironmentVariables(CompilerArguments))
+				if (!IREEUtils::ResolveEnvironmentVariables(CompilerArguments))
 				{
 					UE_LOG(LogIREEDriverRDG, Warning, TEXT("IREECompilerRDG could not replace environment variables in %s"), *Target->CompilerArguments);
 					bResult = false;
@@ -446,7 +315,7 @@ namespace UE::IREE::Compiler
 				CompilerArguments.ReplaceInline(*FString("${VMFB_PATH}"), *(FString("\"") + VmfbFilePath + "\""));
 				CompilerArguments.ReplaceInline(*FString("${INPUT_PATH}"), *(FString("\"") + InputFilePath + "\""));
 
-				RunCommand(CompilerCommand, CompilerArguments, IntermediateFilePathNoExt + "_compile-log.txt");
+				IREEUtils::RunCommand(CompilerCommand, CompilerArguments, IntermediateFilePathNoExt + "_compile-log.txt");
 
 				if (!PlatformFile.FileExists(*VmfbFilePath))
 				{
diff --git a/Engine/Plugins/Experimental/NNERuntimeIREE/Source/IREEUtils/IREEUtils.Build.cs b/Engine/Plugins/Experimental/NNERuntimeIREE/Source/IREEUtils/IREEUtils.Build.cs
new file mode 100644
--- /dev/null
+++ b/Engine/Plugins/Experimental/NNERuntimeIREE/Source/IREEUtils/IREEUtils.Build.cs
@@ -0,0 +1,19 @@
+// Copyright Epic Games, Inc. All Rights Reserved.
+
+using UnrealBuildTool;
+
+public class IREEUtils : ModuleRules
+{
+	public IREEUtils(ReadOnlyTargetRules Target) : base(Target)
+	{
+		PCHUsage = PCHUsageMode.UseExplicitOrSharedPCHs;
+
+		PrivateDependencyModuleNames.AddRange
+		(
+			new string[]
+			{
+				"Core"
+			}
+		);
+	}
+}
diff --git a/Engine/Plugins/Experimental/NNERuntimeIREE/Source/IREEUtils/Internal/IREEUtils.h b/Engine/Plugins/Experimental/NNERuntimeIREE/Source/IREEUtils/Internal/IREEUtils.h
new file mode 100644
--- /dev/null
+++ b/Engine/Plugins/Experimental/NNERuntimeIREE/Source/IREEUtils/Internal/IREEUtils.h
@@ -0,0 +1,16 @@
+// Copyright Epic Games, Inc. All Rights Reserved.
+
+#pragma once
+
+#include "Containers/UnrealString.h"
+
+namespace UE::IREEUtils
+{
+
+IREEUTILS_API bool ResolveEnvironmentVariables(FString& String);
+
+IREEUTILS_API void RunCommand(const FString& Command, const FString& Arguments, const FString& LogFilePath = FString());
+
+IREEUTILS_API bool ImportOnnx(const FString& ImporterCommand, const FString& ImporterArguments, TConstArrayView<uint8> InFileData, const FString& InModelName, const FString& InOutputDir, TArray64<uint8>& OutMlirData);
+
+}
\ No newline at end of file
diff --git a/Engine/Plugins/Experimental/NNERuntimeIREE/Source/IREEUtils/Internal/IREEUtilsLog.h b/Engine/Plugins/Experimental/NNERuntimeIREE/Source/IREEUtils/Internal/IREEUtilsLog.h
new file mode 100644
--- /dev/null
+++ b/Engine/Plugins/Experimental/NNERuntimeIREE/Source/IREEUtils/Internal/IREEUtilsLog.h
@@ -0,0 +1,7 @@
+// Copyright Epic Games, Inc. All Rights Reserved.
+
+#pragma once
+
+#include "Logging/LogMacros.h"
+
+DECLARE_LOG_CATEGORY_EXTERN(LogIREEUtils, Log, All);
\ No newline at end of file
diff --git a/Engine/Plugins/Experimental/NNERuntimeIREE/Source/IREEUtils/Private/IREEUtils.cpp b/Engine/Plugins/Experimental/NNERuntimeIREE/Source/IREEUtils/Private/IREEUtils.cpp
new file mode 100644
--- /dev/null
+++ b/Engine/Plugins/Experimental/NNERuntimeIREE/Source/IREEUtils/Private/IREEUtils.cpp
@@ -0,0 +1,150 @@
+// Copyright Epic Games, Inc. All Rights Reserved.
+
+#include "IREEUtils.h"
+
+#include "HAL/PlatformFileManager.h"
+#include "IREEUtilsLog.h"
+#include "Misc/FileHelper.h"
+#include "Misc/MonitoredProcess.h"
+#include "Misc/Paths.h"
+
+namespace UE::IREEUtils
+{
+
+bool ResolveEnvironmentVariables(FString& String)
+{
+	const FString StartString = "$ENV{";
+	const FString EndString = "}";
+	FString ResultString = String;
+	int32 StartIndex = ResultString.Find(StartString, ESearchCase::CaseSensitive);
+	while (StartIndex != INDEX_NONE)
+	{
+		StartIndex += StartString.Len();
+		int32 EndIndex = ResultString.Find(EndString, ESearchCase::CaseSensitive, ESearchDir::FromStart, StartIndex);
+		if (EndIndex > StartIndex)
+		{
+			FString EnvironmentVariableName = ResultString.Mid(StartIndex, EndIndex - StartIndex);
+			FString EnvironmentVariableValue = FPlatformMisc::GetEnvironmentVariable(*EnvironmentVariableName);
+			if (EnvironmentVariableValue.IsEmpty())
+			{
+				return false;
+			}
+			else
+			{
+				ResultString.ReplaceInline(*(StartString + EnvironmentVariableName + EndString), *EnvironmentVariableValue, ESearchCase::CaseSensitive);
+			}
+		}
+		else
+		{
+			return false;
+		}
+		StartIndex = ResultString.Find(StartString, ESearchCase::CaseSensitive);
+	}
+	String = ResultString;
+	return true;
+}
+
+void RunCommand(const FString& Command, const FString& Arguments, const FString& LogFilePath)
+{
+	int32 ReturnCode = 0;
+	bool IsCanceled = false;
+
+	FMonitoredProcess Process(Command, Arguments, true);
+	Process.OnCompleted().BindLambda([&ReturnCode] (int32 _ReturnCode) { ReturnCode = _ReturnCode; });
+	Process.OnCanceled().BindLambda([&IsCanceled] (){ IsCanceled = true; });
+
+	if (!Process.Launch())
+	{
+		UE_LOG(LogIREEUtils, Warning, TEXT("Failed to launch subprocess!"));
+		return;
+	}
+
+	while (Process.Update())
+	{
+		// Poll until process has finished
+	}
+
+	if (IsCanceled)
+	{
+		UE_LOG(LogIREEUtils, Warning, TEXT("Execution of subprocess was canceled!"));
+	}
+	else if (ReturnCode)
+	{
+		UE_LOG(LogIREEUtils, Warning, TEXT("Subprocess exited with non-zero code %d"), ReturnCode);
+	}
+
+	if (!LogFilePath.IsEmpty())
+	{
+		const FString Output = Process.GetFullOutputWithoutDelegate();
+
+		FStringBuilderBase Builder;
+		Builder.Append(Command)
+			.Append(TEXT(" "))
+			.Append(Arguments)
+			.Append(LINE_TERMINATOR LINE_TERMINATOR)
+			.Append(Output);
+
+		FFileHelper::SaveStringToFile(Builder.ToString(), *LogFilePath);
+
+		UE_LOG(LogIREEUtils, Log, TEXT("Saved subprocess output to: %s"), *LogFilePath);
+	}
+}
+
+bool ImportOnnx(const FString& ImporterCommand, const FString& ImporterArguments, TConstArrayView<uint8> InFileData, const FString& InModelName, const FString& InOutputDir, TArray64<uint8>& OutMlirData)
+{
+	SCOPED_NAMED_EVENT_TEXT("IREEUtils::ImportOnnx", FColor::Magenta);
+
+	using namespace Private;
+
+	IPlatformFile& PlatformFile = FPlatformFileManager::Get().GetPlatformFile();
+
+	FString InputFilePath = FPaths::Combine(InOutputDir, InModelName) + ".onnx";
+	if (!PlatformFile.FileExists(*InputFilePath))
+	{
+		SCOPED_NAMED_EVENT_TEXT("InputFile", FColor::Magenta);
+
+		if(!FFileHelper::SaveArrayToFile(InFileData, *InputFilePath))
+		{
+			UE_LOG(LogIREEUtils, Warning, TEXT("IREECompilerRDG failed to save ONNX model \"%s\""), *InputFilePath);
+			return false;
+		}
+	}
+
+	FString OutputFilePath = FPaths::Combine(InOutputDir, InModelName) + ".mlir";
+	FString IntermediateFilePathNoExt = FPaths::Combine(InOutputDir, InModelName);
+
+	FString ImporterArgumentsCopy = ImporterArguments;
+	if (!IREEUtils::ResolveEnvironmentVariables(ImporterArgumentsCopy))
+	{
+		UE_LOG(LogIREEUtils, Warning, TEXT("IREECompilerRDG could not replace environment variables in %s"), *ImporterArgumentsCopy);
+		return false;
+	}
+	ImporterArgumentsCopy.ReplaceInline(*FString("${INPUT_PATH}"), *(FString("\"") + InputFilePath + "\""));
+	ImporterArgumentsCopy.ReplaceInline(*FString("${OUTPUT_PATH}"), *(FString("\"") + OutputFilePath + "\""));
+
+	{
+		SCOPED_NAMED_EVENT_TEXT("Import", FColor::Magenta);
+
+		IREEUtils::RunCommand(ImporterCommand, ImporterArgumentsCopy, IntermediateFilePathNoExt + "_import-log.txt");
+	}
+
+	if (!PlatformFile.FileExists(*OutputFilePath))
+	{
+		UE_LOG(LogIREEUtils, Warning, TEXT("IREECompilerRDG failed to import the model \"%s\" using the command:"), *InputFilePath);
+		UE_LOG(LogIREEUtils, Warning, TEXT("\"%s\" %s"), *ImporterCommand, *ImporterArgumentsCopy);
+		return false;
+	}
+
+	{
+		SCOPED_NAMED_EVENT_TEXT("Load", FColor::Magenta);
+		if(!FFileHelper::LoadFileToArray(OutMlirData, *OutputFilePath))
+		{
+			UE_LOG(LogIREEUtils, Warning, TEXT("IREECompilerRDG failed to load imported model \"%s\""), *OutputFilePath);
+			return false;
+		}
+	}
+
+	return true;
+}
+
+} // namespace UE::IREEUtils
\ No newline at end of file
diff --git a/Engine/Plugins/Experimental/NNERuntimeIREE/Source/IREEUtils/Private/IREEUtilsModule.cpp b/Engine/Plugins/Experimental/NNERuntimeIREE/Source/IREEUtils/Private/IREEUtilsModule.cpp
new file mode 100644
--- /dev/null
+++ b/Engine/Plugins/Experimental/NNERuntimeIREE/Source/IREEUtils/Private/IREEUtilsModule.cpp
@@ -0,0 +1,16 @@
+// Copyright Epic Games, Inc. All Rights Reserved.
+
+#include "Modules/ModuleManager.h"
+#include "IREEUtilsLog.h"
+
+DEFINE_LOG_CATEGORY(LogIREEUtils);
+
+class FIREEUtilsModule : public IModuleInterface
+{
+public:
+	/** IModuleInterface implementation */
+	virtual void StartupModule() override { }
+	virtual void ShutdownModule() override { }
+};
+
+IMPLEMENT_MODULE(FIREEUtilsModule, IREEUtils)
\ No newline at end of file
diff --git a/Engine/Plugins/Experimental/NNERuntimeIREE/Source/NNERuntimeIREE/NNERuntimeIREE.Build.cs b/Engine/Plugins/Experimental/NNERuntimeIREE/Source/NNERuntimeIREE/NNERuntimeIREE.Build.cs
--- a/Engine/Plugins/Experimental/NNERuntimeIREE/Source/NNERuntimeIREE/NNERuntimeIREE.Build.cs
+++ b/Engine/Plugins/Experimental/NNERuntimeIREE/Source/NNERuntimeIREE/NNERuntimeIREE.Build.cs
@@ -17,6 +17,7 @@ public class NNERuntimeIREE : ModuleRules
 				"CoreUObject",
 				"Engine",
 				"IREE",
+				"IREEUtils",
 				"IREEDriverRDG",
 				"NNE",
 				"Projects",
diff --git a/Engine/Plugins/Experimental/NNERuntimeIREE/Source/NNERuntimeIREE/Private/NNERuntimeIREECompiler.cpp b/Engine/Plugins/Experimental/NNERuntimeIREE/Source/NNERuntimeIREE/Private/NNERuntimeIREECompiler.cpp
--- a/Engine/Plugins/Experimental/NNERuntimeIREE/Source/NNERuntimeIREE/Private/NNERuntimeIREECompiler.cpp
+++ b/Engine/Plugins/Experimental/NNERuntimeIREE/Source/NNERuntimeIREE/Private/NNERuntimeIREECompiler.cpp
@@ -11,9 +11,9 @@
 #include "HAL/FileManager.h"
 #include "HAL/PlatformFileManager.h"
 #include "Interfaces/IPluginManager.h"
+#include "IREEUtils.h"
 #include "Kismet/GameplayStatics.h"
 #include "Misc/FileHelper.h"
-#include "Misc/MonitoredProcess.h"
 #include "Misc/Paths.h"
 #include "NNE.h"
 #include "NNERuntimeIREELog.h"
@@ -27,39 +27,6 @@ namespace UE::NNERuntimeIREE
 	{
 		namespace Private
 		{
-			bool ResolveEnvironmentVariables(FString& String)
-			{
-				const FString StartString = "$ENV{";
-				const FString EndString = "}";
-				FString ResultString = String;
-				int32 StartIndex = ResultString.Find(StartString, ESearchCase::CaseSensitive);
-				while (StartIndex != INDEX_NONE)
-				{
-					StartIndex += StartString.Len();
-					int32 EndIndex = ResultString.Find(EndString, ESearchCase::CaseSensitive, ESearchDir::FromStart, StartIndex);
-					if (EndIndex > StartIndex)
-					{
-						FString EnvironmentVariableName = ResultString.Mid(StartIndex, EndIndex - StartIndex);
-						FString EnvironmentVariableValue = FPlatformMisc::GetEnvironmentVariable(*EnvironmentVariableName);
-						if (EnvironmentVariableValue.IsEmpty())
-						{
-							return false;
-						}
-						else
-						{
-							ResultString.ReplaceInline(*(StartString + EnvironmentVariableName + EndString), *EnvironmentVariableValue, ESearchCase::CaseSensitive);
-						}
-					}
-					else
-					{
-						return false;
-					}
-					StartIndex = ResultString.Find(StartString, ESearchCase::CaseSensitive);
-				}
-				String = ResultString;
-				return true;
-			}
-
 			FString GetSharedLibraryEntryPointName(const FString& HeaderString)
 			{
 				FString SearchString = "iree_hal_executable_library_header_t**";
@@ -76,43 +43,6 @@ namespace UE::NNERuntimeIREE
 				}
 				return HeaderString.Mid(Start, End - Start).TrimStartAndEnd();
 			}
-
-			void RunCommand(const FString& Command, const FString& Arguments, const FString& LogFilePath = FString())
-			{
-				int32 ReturnCode = 0;
-				bool IsCanceled = false;
-
-				FMonitoredProcess Process(Command, Arguments, true);
-				Process.OnCompleted().BindLambda([&ReturnCode] (int32 _ReturnCode) { ReturnCode = _ReturnCode; });
-				Process.OnCanceled().BindLambda([&IsCanceled] (){ IsCanceled = true; });
-
-				if (!Process.Launch())
-				{
-					UE_LOG(LogNNERuntimeIREE, Warning, TEXT("Failed to launch subprocess!"));
-					return;
-				}
-
-				while (Process.Update())
-				{
-					// Poll until process has finished
-				}
-
-				if (IsCanceled)
-				{
-					UE_LOG(LogNNERuntimeIREE, Warning, TEXT("Execution of subprocess was canceled!"));
-				}
-				else if (ReturnCode)
-				{
-					UE_LOG(LogNNERuntimeIREE, Warning, TEXT("Subprocess exited with non-zero code %d"), ReturnCode);
-				}
-
-				if (!LogFilePath.IsEmpty())
-				{
-					FFileHelper::SaveStringToFile(Process.GetFullOutputWithoutDelegate(), *LogFilePath);
-
-					UE_LOG(LogNNERuntimeIREE, Log, TEXT("Saved subprocess output to: %s"), *LogFilePath);
-				}
-			}
 		} // Private
 
 		FCompiler::FCompiler(const FString& InImporterCommand, const FString& InImporterArguments, const FString& InCompilerCommand, const FString& InLinkerCommand, const FString& InSharedLibExt, TConstArrayView<FBuildTarget> InBuildTargets) 
@@ -163,7 +93,7 @@ namespace UE::NNERuntimeIREE
 							FString TmpImporterCommand;
 							for (int32 i = 0; i < BuildConfig.ImporterCommand.Num(); i++)
 							{
-								if (ResolveEnvironmentVariables(BuildConfig.ImporterCommand[i]))
+								if (IREEUtils::ResolveEnvironmentVariables(BuildConfig.ImporterCommand[i]))
 								{
 									BuildConfig.ImporterCommand[i].ReplaceInline(*FString("${PLUGIN_DIR}"), *PluginDir);
 									BuildConfig.ImporterCommand[i].ReplaceInline(*FString("${PROJECT_DIR}"), *FPaths::ProjectDir());
@@ -187,7 +117,7 @@ namespace UE::NNERuntimeIREE
 							FString TmpCompilerCommand;
 							for (int32 i = 0; i < BuildConfig.CompilerCommand.Num(); i++)
 							{
-								if (ResolveEnvironmentVariables(BuildConfig.CompilerCommand[i]))
+								if (IREEUtils::ResolveEnvironmentVariables(BuildConfig.CompilerCommand[i]))
 								{
 									BuildConfig.CompilerCommand[i].ReplaceInline(*FString("${PLUGIN_DIR}"), *PluginDir);
 									BuildConfig.CompilerCommand[i].ReplaceInline(*FString("${PROJECT_DIR}"), *FPaths::ProjectDir());
@@ -211,7 +141,7 @@ namespace UE::NNERuntimeIREE
 							FString TmpLinkerCommand;
 							for (int32 i = 0; i < BuildConfig.LinkerCommand.Num(); i++)
 							{
-								if (ResolveEnvironmentVariables(BuildConfig.LinkerCommand[i]))
+								if (IREEUtils::ResolveEnvironmentVariables(BuildConfig.LinkerCommand[i]))
 								{
 									BuildConfig.LinkerCommand[i].ReplaceInline(*FString("${PLUGIN_DIR}"), *PluginDir);
 									BuildConfig.LinkerCommand[i].ReplaceInline(*FString("${PROJECT_DIR}"), *FPaths::ProjectDir());
@@ -260,58 +190,7 @@ namespace UE::NNERuntimeIREE
 
 		bool FCompiler::ImportOnnx(TConstArrayView<uint8> InFileData, const FString& InModelName, const FString& InOutputDir, TArray64<uint8>& OutMlirData)
 		{
-			SCOPED_NAMED_EVENT_TEXT("FCompiler::ImportOnnx", FColor::Magenta);
-
-			using namespace Private;
-
-			IPlatformFile& PlatformFile = FPlatformFileManager::Get().GetPlatformFile();
-
-			FString InputFilePath = GetOnnxFilePath(InModelName, InOutputDir);
-			if (!PlatformFile.FileExists(*InputFilePath))
-			{
-				SCOPED_NAMED_EVENT_TEXT("InputFile", FColor::Magenta);
-
-				if(!FFileHelper::SaveArrayToFile(InFileData, *InputFilePath))
-				{
-					UE_LOG(LogNNERuntimeIREE, Warning, TEXT("UNNERuntimeIREECpu failed to save ONNX model \"%s\""), *InputFilePath);
-					return false;
-				}
-			}
-
-			FString OutputFilePath = GetMlirFilePath(InModelName, InOutputDir);
-			FString IntermediateFilePathNoExt = FPaths::Combine(InOutputDir, InModelName);
-
-			if (!ResolveEnvironmentVariables(ImporterArguments))
-			{
-				UE_LOG(LogNNERuntimeIREE, Warning, TEXT("UNNERuntimeIREECpu could not replace environment variables in %s"), *ImporterArguments);
-				return false;
-			}
-			ImporterArguments.ReplaceInline(*FString("${INPUT_PATH}"), *(FString("\"") + InputFilePath + "\""));
-			ImporterArguments.ReplaceInline(*FString("${OUTPUT_PATH}"), *(FString("\"") + OutputFilePath + "\""));
-
-			{
-				SCOPED_NAMED_EVENT_TEXT("Import", FColor::Magenta);
-
-				RunCommand(ImporterCommand, ImporterArguments, IntermediateFilePathNoExt + "_import-log.txt");
-			}
-
-			if (!PlatformFile.FileExists(*OutputFilePath))
-			{
-				UE_LOG(LogNNERuntimeIREE, Warning, TEXT("UNNERuntimeIREECpu failed to import the model \"%s\" using the command:"), *InputFilePath);
-				UE_LOG(LogNNERuntimeIREE, Warning, TEXT("\"%s\" %s"), *ImporterCommand, *ImporterArguments);
-				return false;
-			}
-
-			{
-				SCOPED_NAMED_EVENT_TEXT("Load", FColor::Magenta);
-				if(!FFileHelper::LoadFileToArray(OutMlirData, *OutputFilePath))
-				{
-					UE_LOG(LogNNERuntimeIREE, Warning, TEXT("UNNERuntimeIREECpu failed to load imported model \"%s\""), *OutputFilePath);
-					return false;
-				}
-			}
-
-			return true;
+			return IREEUtils::ImportOnnx(ImporterCommand, ImporterArguments, InFileData, InModelName, InOutputDir, OutMlirData);
 		}
 
 		bool FCompiler::CompileMlir(TConstArrayView<uint8> InFileData, const FString& InModelName, const FString& InOutputDir, FNNERuntimeIREECompilerResultCPU& OutCompilerResult, UNNERuntimeIREEModuleMetaData& ModuleMetaData)
@@ -330,7 +209,7 @@ namespace UE::NNERuntimeIREE
 
 			IPlatformFile& PlatformFile = FPlatformFileManager::Get().GetPlatformFile();
 
-			FString InputFilePath = GetMlirFilePath(InModelName, InOutputDir);
+			FString InputFilePath = FPaths::Combine(InOutputDir, InModelName) + ".mlir";
 			if (!PlatformFile.FileExists(*InputFilePath))
 			{
 				SCOPED_NAMED_EVENT_TEXT("InputFile", FColor::Magenta);
@@ -349,7 +228,7 @@ namespace UE::NNERuntimeIREE
 				FString SharedLibFilePath = IntermediateFilePathNoExt + SharedLibExt;
 
 				FString CompilerArguments = BuildTargets[i].CompilerArguments;
-				if (!ResolveEnvironmentVariables(CompilerArguments))
+				if (!IREEUtils::ResolveEnvironmentVariables(CompilerArguments))
 				{
 					UE_LOG(LogNNERuntimeIREE, Warning, TEXT("UNNERuntimeIREECpu could not replace environment variables in %s"), *BuildTargets[i].CompilerArguments);
 					bResult = false;
@@ -362,7 +241,7 @@ namespace UE::NNERuntimeIREE
 				{
 					SCOPED_NAMED_EVENT_TEXT("Compile", FColor::Magenta);
 
-					RunCommand(CompilerCommand, CompilerArguments, IntermediateFilePathNoExt + "_compile-log.txt");
+					IREEUtils::RunCommand(CompilerCommand, CompilerArguments, IntermediateFilePathNoExt + "_compile-log.txt");
 				}
 
 				if (!PlatformFile.FileExists(*ObjectFilePath) || !PlatformFile.FileExists(*VmfbFilePath))
@@ -374,7 +253,7 @@ namespace UE::NNERuntimeIREE
 				}
 
 				FString LinkerArguments = BuildTargets[i].LinkerArguments;
-				if (!ResolveEnvironmentVariables(LinkerArguments))
+				if (!IREEUtils::ResolveEnvironmentVariables(LinkerArguments))
 				{
 					UE_LOG(LogNNERuntimeIREE, Warning, TEXT("UNNERuntimeIREECpu could not replace environment variables in %s"), *BuildTargets[i].LinkerArguments);
 					bResult = false;
@@ -386,7 +265,7 @@ namespace UE::NNERuntimeIREE
 				{
 					SCOPED_NAMED_EVENT_TEXT("Link", FColor::Magenta);
 
-					RunCommand(LinkerCommand, LinkerArguments, IntermediateFilePathNoExt + "_link-log.txt");
+					IREEUtils::RunCommand(LinkerCommand, LinkerArguments, IntermediateFilePathNoExt + "_link-log.txt");
 				}
 
 				if (!PlatformFile.FileExists(*SharedLibFilePath))
@@ -438,16 +317,6 @@ namespace UE::NNERuntimeIREE
 
 			return bResult;
 		}
-
-		FString FCompiler::GetOnnxFilePath(const FString& InModelName, const FString& InOutputDir)
-		{
-			return FPaths::Combine(InOutputDir, InModelName) + ".onnx";
-		}
-
-		FString FCompiler::GetMlirFilePath(const FString& InModelName, const FString& InOutputDir)
-		{
-			return FPaths::Combine(InOutputDir, InModelName) + ".mlir";
-		}
 	} // CPU
 } // UE::NNERuntimeIREE
 
diff --git a/Engine/Plugins/Experimental/NNERuntimeIREE/Source/NNERuntimeIREE/Private/NNERuntimeIREECompiler.h b/Engine/Plugins/Experimental/NNERuntimeIREE/Source/NNERuntimeIREE/Private/NNERuntimeIREECompiler.h
--- a/Engine/Plugins/Experimental/NNERuntimeIREE/Source/NNERuntimeIREE/Private/NNERuntimeIREECompiler.h
+++ b/Engine/Plugins/Experimental/NNERuntimeIREE/Source/NNERuntimeIREE/Private/NNERuntimeIREECompiler.h
@@ -89,9 +89,6 @@ namespace UE::NNERuntimeIREE::CPU
 		bool CompileMlir(TConstArrayView<uint8> InFileData, const FString& InModelName, const FString& InOutputDir, FNNERuntimeIREECompilerResultCPU& OutCompilerResult, UNNERuntimeIREEModuleMetaData& ModuleMetaData);
 
 	private:
-		static inline FString GetOnnxFilePath(const FString& InModelName, const FString& InOutputDir);
-		static inline FString GetMlirFilePath(const FString& InModelName, const FString& InOutputDir);
-
 		FString ImporterCommand;
 		FString ImporterArguments;
 		FString CompilerCommand;
