Disable NNERuntimeIREE plugin for Windows Arm64 for now

#rnx

diff --git a/Engine/Plugins/Experimental/NNERuntimeIREE/NNERuntimeIREE.uplugin b/Engine/Plugins/Experimental/NNERuntimeIREE/NNERuntimeIREE.uplugin
--- a/Engine/Plugins/Experimental/NNERuntimeIREE/NNERuntimeIREE.uplugin
+++ b/Engine/Plugins/Experimental/NNERuntimeIREE/NNERuntimeIREE.uplugin
@@ -18,27 +18,32 @@
 		{
 			"Name": "IREEUtils",
 			"Type": "Runtime",
-			"LoadingPhase": "Default"
+			"LoadingPhase": "Default",
+			"PlatformArchitectureDenyList": [ "Win64:arm64" ]
 		},
 		{
 			"Name": "IREEDriverRDG",
 			"Type": "Runtime",
-			"LoadingPhase": "Default"
+			"LoadingPhase": "Default",
+			"PlatformArchitectureDenyList": [ "Win64:arm64" ]
 		},
 		{
 			"Name": "NNERuntimeIREE",
 			"Type": "Runtime",
-			"LoadingPhase": "Default"
+			"LoadingPhase": "Default",
+			"PlatformArchitectureDenyList": [ "Win64:arm64" ]
 		},
 		{
 			"Name": "NNERuntimeIREEEditor",
 			"Type": "Editor",
-			"LoadingPhase": "Default"
+			"LoadingPhase": "Default",
+			"PlatformArchitectureDenyList": [ "Win64:arm64" ]
 		},
 		{
 			"Name": "NNERuntimeIREEShader",
 			"Type": "Runtime",
-			"LoadingPhase": "PostConfigInit"
+			"LoadingPhase": "PostConfigInit",
+			"PlatformArchitectureDenyList": [ "Win64:arm64" ]
 		}
 	]
 }
\ No newline at end of file
