[NNE]  NNERuntimeIREECpu added multi-task support backed by IREE CPU local-task device.

#tests Editor, Game on all supported platforms

diff --git a/Engine/Plugins/Experimental/NNERuntimeIREE/Source/NNERuntimeIREE/NNERuntimeIREE.Build.cs b/Engine/Plugins/Experimental/NNERuntimeIREE/Source/NNERuntimeIREE/NNERuntimeIREE.Build.cs
--- a/Engine/Plugins/Experimental/NNERuntimeIREE/Source/NNERuntimeIREE/NNERuntimeIREE.Build.cs
+++ b/Engine/Plugins/Experimental/NNERuntimeIREE/Source/NNERuntimeIREE/NNERuntimeIREE.Build.cs
@@ -15,6 +15,7 @@ public class NNERuntimeIREE : ModuleRules
 			{
 				"Core",
 				"CoreUObject",
+				"DeveloperSettings",
 				"Engine",
 				"IREE",
 				"IREEDriverRDG",
diff --git a/Engine/Plugins/Experimental/NNERuntimeIREE/Source/NNERuntimeIREE/Private/NNERuntimeIREE.cpp b/Engine/Plugins/Experimental/NNERuntimeIREE/Source/NNERuntimeIREE/Private/NNERuntimeIREE.cpp
--- a/Engine/Plugins/Experimental/NNERuntimeIREE/Source/NNERuntimeIREE/Private/NNERuntimeIREE.cpp
+++ b/Engine/Plugins/Experimental/NNERuntimeIREE/Source/NNERuntimeIREE/Private/NNERuntimeIREE.cpp
@@ -142,6 +142,11 @@ namespace UE::NNERuntimeIREE::Private
 FGuid UNNERuntimeIREECpu::GUID = FGuid((int32)'I', (int32)'C', (int32)'P', (int32)'U');
 int32 UNNERuntimeIREECpu::Version = 0x00000005;
 
+void UNNERuntimeIREECpu::Init(TSharedRef<UE::NNERuntimeIREE::Private::FEnvironment> InEnvironment)
+{
+	Environment = InEnvironment;
+}
+
 FString UNNERuntimeIREECpu::GetRuntimeName() const
 {
 	return TEXT("NNERuntimeIREECpu");
@@ -414,7 +419,7 @@ TSharedPtr<UE::NNE::IModelCPU> UNNERuntimeIREECpu::CreateModelCPU(const TObjectP
 	const FString SharedLibraryDirFullPath = FPaths::ConvertRelativePathToFull(FPaths::Combine(FPaths::ProjectDir(), GetSharedLibDirPath(UGameplayStatics::GetPlatformName(), GetRuntimeName(), FileIdString)));
 	const FString SharedLibrarySubDirFullPath = FPaths::Combine(SharedLibraryDirFullPath, ArchitectureInfo.RelativeDirPath);
 
-	TSharedPtr<UE::NNE::IModelCPU> Model = UE::NNERuntimeIREE::CPU::FModel::Make(SharedLibrarySubDirFullPath, ArchitectureInfo.SharedLibraryFileName, ArchitectureInfo.VmfbFileName, ArchitectureInfo.SharedLibraryEntryPointName, *ModuleMetaData);
+	TSharedPtr<UE::NNE::IModelCPU> Model = UE::NNERuntimeIREE::CPU::FModel::Make(*Environment, SharedLibrarySubDirFullPath, ArchitectureInfo.SharedLibraryFileName, ArchitectureInfo.VmfbFileName, ArchitectureInfo.SharedLibraryEntryPointName, *ModuleMetaData);
 	if (!Model.IsValid())
 	{
 		UE_LOG(LogNNERuntimeIREE, Warning, TEXT("UNNERuntimeIREECpu could not initialize the model created from model data with id %s"), *FileIdString);
diff --git a/Engine/Plugins/Experimental/NNERuntimeIREE/Source/NNERuntimeIREE/Private/NNERuntimeIREE.h b/Engine/Plugins/Experimental/NNERuntimeIREE/Source/NNERuntimeIREE/Private/NNERuntimeIREE.h
--- a/Engine/Plugins/Experimental/NNERuntimeIREE/Source/NNERuntimeIREE/Private/NNERuntimeIREE.h
+++ b/Engine/Plugins/Experimental/NNERuntimeIREE/Source/NNERuntimeIREE/Private/NNERuntimeIREE.h
@@ -12,6 +12,13 @@
 
 #include "NNERuntimeIREE.generated.h"
 
+#ifdef WITH_NNE_RUNTIME_IREE
+namespace UE::NNERuntimeIREE::Private
+{
+	class FEnvironment;
+}
+#endif // WITH_NNE_RUNTIME_IREE
+
 UCLASS()
 class UNNERuntimeIREECpu : public UObject, public INNERuntime, public INNERuntimeCPU
 {
@@ -22,6 +29,8 @@ public:
 	static FGuid GUID;
 	static int32 Version;
 
+	void Init(TSharedRef<UE::NNERuntimeIREE::Private::FEnvironment> InEnvironment);
+
 	//~ Begin INNERuntime Interface
 	virtual FString GetRuntimeName() const override;
 	virtual ECanCreateModelDataStatus CanCreateModelData(const FString& FileType, TConstArrayView64<uint8> FileData, const TMap<FString, TConstArrayView64<uint8>>& AdditionalFileData, const FGuid& FileId, const ITargetPlatform* TargetPlatform) const override;
@@ -33,6 +42,10 @@ public:
 	virtual ECanCreateModelCPUStatus CanCreateModelCPU(const TObjectPtr<UNNEModelData> ModelData) const override;
 	virtual TSharedPtr<UE::NNE::IModelCPU> CreateModelCPU(const TObjectPtr<UNNEModelData> ModelData) override;
 	//~ End INNERuntimeCPU Interface
+
+private:
+	TSharedPtr<UE::NNERuntimeIREE::Private::FEnvironment> Environment;
+
 #else
 	//~ Begin INNERuntime Interface
 	virtual FString GetRuntimeName() const override { return ""; };
diff --git a/Engine/Plugins/Experimental/NNERuntimeIREE/Source/NNERuntimeIREE/Private/NNERuntimeIREEEnvironment.cpp b/Engine/Plugins/Experimental/NNERuntimeIREE/Source/NNERuntimeIREE/Private/NNERuntimeIREEEnvironment.cpp
new file mode 100644
--- /dev/null
+++ b/Engine/Plugins/Experimental/NNERuntimeIREE/Source/NNERuntimeIREE/Private/NNERuntimeIREEEnvironment.cpp
@@ -0,0 +1,462 @@
+// Copyright Epic Games, Inc. All Rights Reserved.
+
+#include "NNERuntimeIREEEnvironment.h"
+
+#ifdef WITH_NNE_RUNTIME_IREE
+
+#include "Algo/AllOf.h"
+#include "Misc/ScopeLock.h"
+
+#if PLATFORM_MICROSOFT
+#include "Microsoft/AllowMicrosoftPlatformTypes.h"
+#include "Microsoft/AllowMicrosoftPlatformAtomics.h"
+#endif // PLATFORM_MICROSOFT
+THIRD_PARTY_INCLUDES_START
+#include "iree/task/api.h"
+#include "iree/task/executor.h"
+#include "iree/task/topology.h"
+THIRD_PARTY_INCLUDES_END
+#if PLATFORM_MICROSOFT
+#include "Microsoft/HideMicrosoftPlatformAtomics.h"
+#include "Microsoft/HideMicrosoftPlatformTypes.h"
+#endif // PLATFORM_MICROSOFT
+
+LLM_DEFINE_TAG(NNERuntimeIREE_Global);
+
+namespace UE::NNERuntimeIREE::Private
+{
+
+namespace HostAllocatorUtils
+{
+
+iree_status_t HostAllocatorImplAlloc(iree_allocator_command_t Command, const iree_allocator_alloc_params_t* Params, void** InoutPtr)
+{
+	check(Params);
+	check(InoutPtr);
+	switch (Command)
+	{
+		case IREE_ALLOCATOR_COMMAND_MALLOC:
+			*InoutPtr = FMemory::Malloc(Params->byte_length, IREE_HAL_HEAP_BUFFER_ALIGNMENT);
+			break;
+		case IREE_ALLOCATOR_COMMAND_CALLOC:
+			*InoutPtr = FMemory::MallocZeroed(Params->byte_length, IREE_HAL_HEAP_BUFFER_ALIGNMENT);
+			break;
+		case IREE_ALLOCATOR_COMMAND_REALLOC:
+			*InoutPtr = FMemory::Realloc(*InoutPtr, Params->byte_length, IREE_HAL_HEAP_BUFFER_ALIGNMENT);
+			break;
+		default:
+			check(false);
+	}
+	if (!*InoutPtr)
+	{
+		return iree_make_status(IREE_STATUS_RESOURCE_EXHAUSTED, "Host allocator failed during alloc");
+	}
+	return iree_ok_status();
+}
+
+iree_status_t HostAllocatorImpl(void *Self, iree_allocator_command_t Command, const void* Params, void** InoutPtr)
+{
+	LLM_SCOPE_BYTAG(NNERuntimeIREE_Global);
+	check(InoutPtr);
+	switch (Command)
+	{
+		case IREE_ALLOCATOR_COMMAND_MALLOC:
+		case IREE_ALLOCATOR_COMMAND_CALLOC:
+		case IREE_ALLOCATOR_COMMAND_REALLOC:
+			return HostAllocatorImplAlloc(Command, static_cast<const iree_allocator_alloc_params_t*>(Params), InoutPtr);
+		case IREE_ALLOCATOR_COMMAND_FREE:
+			FMemory::Free(*InoutPtr);
+			break;
+		default:
+			return iree_make_status(IREE_STATUS_UNIMPLEMENTED, "Host allocator command not recognised");
+	}
+	return iree_ok_status();
+}
+
+} // namespace HostAllocatorUtils
+
+iree_allocator_t MakeHostAllocator()
+{
+	iree_allocator_t Allocator = { NULL, HostAllocatorUtils::HostAllocatorImpl };
+	return Allocator;
+}
+
+void PrintIREEError(const FString& InMessage, iree_status_t InStatus)
+{
+	iree_host_size_t TrueLength = 0;
+	iree_status_format(InStatus, 0, (char*)nullptr, &TrueLength);
+	void* ErrorString = FMemory::Malloc(TrueLength + 1);
+	((char*)ErrorString)[TrueLength] = (char)0;
+	iree_status_format(InStatus, TrueLength, (char*)ErrorString, &TrueLength);
+	UE_LOG(LogNNERuntimeIREE, Error, TEXT("%s: %s"), *InMessage, *FString(StringCast<TCHAR>(static_cast<const ANSICHAR*>(ErrorString)).Get()));
+	FMemory::Free(ErrorString);
+}
+
+namespace TaskUtils
+{
+
+bool CreateTaskTopology(const TArray<FNNERuntimeIREEThreadingAffinity>& TaskTopologyGroups, TArray<iree_task_topology_t>& TaskTopologies)
+{
+	if (TaskTopologyGroups.IsEmpty())
+	{
+		UE_LOG(LogNNERuntimeIREE, Log, TEXT("Using default task topology..."));
+
+		const iree_host_size_t MaxCoreCount = 64;
+		iree_task_topology_t TaskTopology;
+		iree_status_t Status = iree_task_topology_initialize_from_physical_cores(IREE_TASK_TOPOLOGY_NODE_ID_ANY, IREE_TASK_TOPOLOGY_PERFORMANCE_LEVEL_ANY, MaxCoreCount, &TaskTopology);
+		if (!iree_status_is_ok(Status))
+		{
+			Private::PrintIREEError("Failed to initialize task topology", Status);
+			return false;
+		}
+
+		TaskTopologies.Add(TaskTopology);
+
+		return true;
+	}
+
+	const iree_host_size_t NumNodes = iree_task_topology_query_node_count();
+	const iree_task_topology_node_id_t CurrentNode = iree_task_topology_query_current_node();
+
+	bool bSuccess = true;
+
+	if (Algo::AllOf(TaskTopologyGroups, [](FNNERuntimeIREEThreadingAffinity Group) { return Group.CoreIndex == -1; }))
+	{
+		check(!TaskTopologyGroups.IsEmpty());
+
+		const int32 NumCores = TaskTopologyGroups.Num();
+		const ENNERuntimeIREEThreadingAffinityGroupSpecifierType Type = TaskTopologyGroups[0].GroupSpecifierType;
+		const int32 GroupIndex = TaskTopologyGroups[0].GroupIndex;
+
+		// To keep implementation simple, we are quite restrictive for now (similar to IREE)
+		if (!Algo::AllOf(TaskTopologyGroups, [Type](FNNERuntimeIREEThreadingAffinity Group) { return Group.GroupSpecifierType == Type; }))
+		{
+			UE_LOG(LogNNERuntimeIREE, Warning, TEXT("All task topology groups require the same specifier type"));
+			return false;
+		}
+		if (Type == ENNERuntimeIREEThreadingAffinityGroupSpecifierType::Index &&
+			!Algo::AllOf(TaskTopologyGroups, [GroupIndex](FNNERuntimeIREEThreadingAffinity Group) { return Group.GroupIndex == GroupIndex; }))
+		{
+			UE_LOG(LogNNERuntimeIREE, Warning, TEXT("All task topology groups require the same NUMA node index"));
+			return false;
+		}
+
+		const int32 TotalCores = Type == ENNERuntimeIREEThreadingAffinityGroupSpecifierType::All ? NumCores * NumNodes : NumCores;
+		const iree_host_size_t MaxCoreCount = TotalCores == 0 ? 64 : TotalCores;
+
+		const uint32 NodeId = Type == ENNERuntimeIREEThreadingAffinityGroupSpecifierType::Index ?
+			GroupIndex :
+				Type == ENNERuntimeIREEThreadingAffinityGroupSpecifierType::Current ?
+				CurrentNode : IREE_TASK_TOPOLOGY_NODE_ID_ANY;
+
+		iree_task_topology_t TaskTopology;
+		iree_status_t Status = iree_task_topology_initialize_from_physical_cores(NodeId, IREE_TASK_TOPOLOGY_PERFORMANCE_LEVEL_ANY, MaxCoreCount, &TaskTopology);
+		if (!iree_status_is_ok(Status))
+		{
+			Private::PrintIREEError("Failed to initialize task topology", Status);
+			return false;
+		}
+
+		TaskTopologies.Add(TaskTopology);
+	}
+	else if (Algo::AllOf(TaskTopologyGroups, [](FNNERuntimeIREEThreadingAffinity Group) { return Group.CoreIndex > -1; }))
+	{
+		TArray<iree_thread_affinity_t> ThreadAffinities;
+		for (const FNNERuntimeIREEThreadingAffinity& Group : TaskTopologyGroups)
+		{
+			TArray<uint32> Nodes;
+			// note: if we are allowed to place threads on any NUMA node, we place it on the same as calling thread (same as type 'current')
+			if (Group.GroupSpecifierType == ENNERuntimeIREEThreadingAffinityGroupSpecifierType::Current || Group.GroupSpecifierType == ENNERuntimeIREEThreadingAffinityGroupSpecifierType::Any)
+			{
+				Nodes.Add(CurrentNode);
+			}
+			else if (Group.GroupSpecifierType == ENNERuntimeIREEThreadingAffinityGroupSpecifierType::All)
+			{
+				for (int i = 0; i < NumNodes; ++i)
+				{
+					Nodes.Add(i);
+				}
+			}
+			else
+			{
+				check(Group.GroupSpecifierType == ENNERuntimeIREEThreadingAffinityGroupSpecifierType::Index);
+				Nodes.Add(Group.GroupIndex);
+			}
+
+			for (uint32 n = 0; n < (uint32)Nodes.Num(); ++n)
+			{
+				check(Group.CoreIndex >= 0);
+
+				ThreadAffinities.Add(iree_thread_affinity_t
+					{
+						.specified = 1u,
+						.smt = 1,
+						.group = Nodes[n],
+						.id = static_cast<uint32>(Group.CoreIndex)
+					});
+			}
+		}
+
+		iree_task_topology_t TaskTopology;
+		iree_status_t Status = iree_task_topology_initialize_from_thread_affinities(ThreadAffinities.Num(), ThreadAffinities.GetData(), &TaskTopology);
+		if (!iree_status_is_ok(Status))
+		{
+			Private::PrintIREEError("Failed to initialize task topology from thread affinities: ", Status);
+			bSuccess = false;
+		}
+		else
+		{
+			TaskTopologies.Add(TaskTopology);
+		}
+	}
+	else
+	{
+		UE_LOG(LogNNERuntimeIREE, Warning, TEXT("Inhomogeneous core assignment not supported yet"));
+	}
+
+	if (!bSuccess)
+	{
+		for (auto TaskTopology : TaskTopologies)
+		{
+			iree_task_topology_deinitialize(&TaskTopology);
+		}
+		TaskTopologies.Empty();
+	}
+
+	if (!bSuccess || TaskTopologies.IsEmpty())
+	{
+		return false;
+	}
+
+	return true;
+}
+
+#define DUMP_TASK_TOPOLOGY_LOG_CATEGORY Verbose
+
+// Ported 'iree_task_flags_print_action_flag()' from runtime\src\iree\task\api.c to use logging in Unreal Engine
+static void DumpTaskTopology(iree_host_size_t TopologyId, const iree_task_topology_t* Topology)
+{
+	UE_LOG(LogNNERuntimeIREE, DUMP_TASK_TOPOLOGY_LOG_CATEGORY, TEXT("# ===--------------------------------------------------------------==="));
+	UE_LOG(LogNNERuntimeIREE, DUMP_TASK_TOPOLOGY_LOG_CATEGORY, TEXT("# topology[%llu]: %llu worker groups"), static_cast<uint64>(TopologyId), static_cast<uint64>(Topology->group_count));
+	UE_LOG(LogNNERuntimeIREE, DUMP_TASK_TOPOLOGY_LOG_CATEGORY, TEXT("# ===--------------------------------------------------------------===\n#"));
+
+	for (iree_host_size_t j = 0; j < Topology->group_count; ++j)
+	{
+		const iree_task_topology_group_t* Group = &Topology->groups[j];
+
+		UE_LOG(LogNNERuntimeIREE, DUMP_TASK_TOPOLOGY_LOG_CATEGORY, TEXT("# group[%d]: '%s'"), Group->group_index, UTF8_TO_TCHAR(Group->name));
+		UE_LOG(LogNNERuntimeIREE, DUMP_TASK_TOPOLOGY_LOG_CATEGORY, TEXT("#      processor: %u"), Group->processor_index);
+
+		if (Group->ideal_thread_affinity.specified)
+		{
+			UE_LOG(LogNNERuntimeIREE, DUMP_TASK_TOPOLOGY_LOG_CATEGORY, TEXT("#       affinity: group=%u, id=%u, smt=%u"), Group->ideal_thread_affinity.group, Group->ideal_thread_affinity.id, Group->ideal_thread_affinity.smt);
+		}
+		else
+		{
+			UE_LOG(LogNNERuntimeIREE, DUMP_TASK_TOPOLOGY_LOG_CATEGORY, TEXT("#       affinity: (unspecified)"));
+		}
+
+		UE_LOG(LogNNERuntimeIREE, DUMP_TASK_TOPOLOGY_LOG_CATEGORY, TEXT("#  caches: l1d=%u, l2d=%u"), Group->caches.l1_data, Group->caches.l2_data);
+
+		if (Group->constructive_sharing_mask == 0)
+		{
+			UE_LOG(LogNNERuntimeIREE, DUMP_TASK_TOPOLOGY_LOG_CATEGORY, TEXT("#  last level cache sharing: (none)"));
+		}
+		else if (Group->constructive_sharing_mask == IREE_TASK_TOPOLOGY_GROUP_MASK_ALL)
+		{
+			UE_LOG(LogNNERuntimeIREE, DUMP_TASK_TOPOLOGY_LOG_CATEGORY, TEXT("#  last level cache sharing: (all/undefined)"));
+		}
+		else
+		{
+			const uint32 ShareCount = iree_math_count_ones_u64(Group->constructive_sharing_mask);
+
+			FString MaskString;
+			for (iree_host_size_t ic = 0, jc = 0; ic < IREE_TASK_TOPOLOGY_GROUP_BIT_COUNT; ++ic)
+			{
+				if ((Group->constructive_sharing_mask >> ic) & 1)
+				{
+					if (jc++ > 0)
+					{
+						MaskString.Append(TEXT(", "));
+					}
+
+					MaskString.Append(FString::Printf(TEXT("%llu"), static_cast<uint64>(ic)));
+				}
+			}
+
+			UE_LOG(LogNNERuntimeIREE, DUMP_TASK_TOPOLOGY_LOG_CATEGORY, TEXT("#  last level cache sharing: %u group(s): %s"), ShareCount, *MaskString);
+		}
+
+		UE_LOG(LogNNERuntimeIREE, DUMP_TASK_TOPOLOGY_LOG_CATEGORY, TEXT("#"));
+	}
+}
+
+} // namespace TaskUtils
+
+// Unique wrapper around iree_task_executor*
+class FTaskExecutor
+{
+public:
+	explicit FTaskExecutor(iree_task_executor_t* TaskExecutor) : TaskExecutor(TaskExecutor)
+	{
+		check(TaskExecutor != nullptr);
+		iree_task_executor_retain(TaskExecutor);
+	}
+
+	~FTaskExecutor()
+	{
+		if (TaskExecutor)
+		{
+			iree_task_executor_release(TaskExecutor);
+		}
+	}
+
+	// Disallow copying
+	FTaskExecutor(const FTaskExecutor&) = delete;
+	FTaskExecutor& operator=(const FTaskExecutor&) = delete;
+
+	// Moving
+	FTaskExecutor(FTaskExecutor&& Other) noexcept : TaskExecutor(Other.TaskExecutor)
+	{
+		Other.TaskExecutor = nullptr;
+	}
+
+	FTaskExecutor& operator=(FTaskExecutor&& Other) noexcept
+	{
+		if (this != &Other)
+		{
+			if (TaskExecutor)
+			{
+				iree_task_executor_release(TaskExecutor);
+			}
+
+			TaskExecutor = Other.TaskExecutor;
+			Other.TaskExecutor = nullptr;
+		}
+		return *this;
+	}
+
+	iree_task_executor_t* Get() const
+	{
+		return TaskExecutor;
+	}
+
+private:
+	iree_task_executor_t* TaskExecutor;
+};
+
+struct FEnvironment::FInternalState
+{
+	TArray<FTaskExecutor> TaskExecutors;
+};
+
+FEnvironment::FEnvironment() = default;
+FEnvironment::~FEnvironment() = default;
+
+void FEnvironment::Configure(const FConfig& InConfig)
+{
+	FScopeLock ScopeLock(&CriticalSection);
+
+#if WITH_EDITOR
+	// For reloading need to make sure thread pool is release before we create a new one
+	InternalState.Reset();
+#endif // WITH_EDITOR
+
+	checkf(!InternalState.IsValid(), TEXT("IREE environment already created!"));
+
+	Config = InConfig;
+}
+
+TArray<iree_task_executor_t*> FEnvironment::GetTaskExecutors() const
+{
+	FScopeLock ScopeLock(&CriticalSection);
+
+	if (!InternalState.IsValid())
+	{
+		if (!Create())
+		{
+			return {};
+		}
+	}
+	
+	check(InternalState.IsValid());
+	check(!InternalState->TaskExecutors.IsEmpty());
+
+	TArray<iree_task_executor_t*> Result;
+	Result.SetNumUninitialized(InternalState->TaskExecutors.Num());
+
+	for (int32 i = 0; i < InternalState->TaskExecutors.Num(); i++)
+	{
+		Result[i] = InternalState->TaskExecutors[i].Get();
+	}
+
+	return Result;
+}
+
+bool FEnvironment::Create() const
+{
+	checkf(!Config.ThreadingOptions.bIsSingleThreaded, TEXT("Can not request task executors in single threaded mode!"));
+	checkf(!InternalState.IsValid(), TEXT("IREE environment already created!"));
+	
+	TArray<iree_task_topology_t> TaskTopologies;
+	if (!TaskUtils::CreateTaskTopology(Config.ThreadingOptions.TaskTopology.TaskTopologyGroups, TaskTopologies))
+	{
+		return false;
+	}
+	
+	check(!TaskTopologies.IsEmpty());
+
+	iree_allocator_t HostAllocator = MakeHostAllocator();
+
+	iree_task_executor_options_t TaskExecutorOptions;
+	iree_task_executor_options_initialize(&TaskExecutorOptions);
+
+	iree_status_t Status = iree_ok_status();
+
+	TArray<FTaskExecutor> TaskExecutors;
+	TaskExecutors.SetNumZeroed(TaskTopologies.Num());
+	for (int i = 0; i < TaskTopologies.Num(); ++i)
+	{
+		TaskUtils::DumpTaskTopology(i, &TaskTopologies[i]);
+
+		iree_task_executor_t* TaskExecutor = nullptr;
+		Status = iree_task_executor_create(TaskExecutorOptions, &TaskTopologies[i], HostAllocator, &TaskExecutor);
+		if (!iree_status_is_ok(Status)) {
+			Private::PrintIREEError("Failed to initialize task executor", Status);
+			break;
+		}
+
+		TaskExecutors[i] = FTaskExecutor(TaskExecutor);
+
+		// task executor is managed by wrapper class
+		iree_task_executor_release(TaskExecutor);
+	}
+
+	// Now we don't need the topologies anymore
+	for (iree_task_topology_t TaskTopology : TaskTopologies)
+	{
+		iree_task_topology_deinitialize(&TaskTopology);
+	}
+
+	if (!iree_status_is_ok(Status))
+	{
+		// Note:
+		// - we already logged the error status
+		// - already successfully created task executors are in local variable 'TaskExecutors' which on return will be destroyed and releases them properly
+		return false;
+	}
+
+	InternalState = MakeUnique<FInternalState>();
+	InternalState->TaskExecutors = MoveTemp(TaskExecutors);
+
+	// 
+	// WITH_EDITOR After this point, if we created a new thread pool, all existing ModelInstances should be created again to use the new one
+	// note: they still can use the old ones because iree devices hold a reference to the iree task exectutors
+	// 
+
+	return true;
+}
+
+} // UE::NNERuntimeIREE::Private
+
+#endif // WITH_NNE_RUNTIME_IREE
\ No newline at end of file
diff --git a/Engine/Plugins/Experimental/NNERuntimeIREE/Source/NNERuntimeIREE/Private/NNERuntimeIREEEnvironment.h b/Engine/Plugins/Experimental/NNERuntimeIREE/Source/NNERuntimeIREE/Private/NNERuntimeIREEEnvironment.h
new file mode 100644
--- /dev/null
+++ b/Engine/Plugins/Experimental/NNERuntimeIREE/Source/NNERuntimeIREE/Private/NNERuntimeIREEEnvironment.h
@@ -0,0 +1,55 @@
+// Copyright Epic Games, Inc. All Rights Reserved.
+
+#pragma once
+
+#ifdef WITH_NNE_RUNTIME_IREE
+
+#include "HAL/CriticalSection.h"
+#include "NNERuntimeIREELog.h"
+#include "NNERuntimeIREESettings.h"
+
+struct iree_task_executor_t;
+struct iree_allocator_t;
+typedef struct iree_status_handle_t* iree_status_t;
+
+namespace UE::NNERuntimeIREE::Private
+{
+
+// Make IREE host allocator that uses Unreal Engine FMemory
+iree_allocator_t MakeHostAllocator();
+
+// Print IREE error formatted
+void PrintIREEError(const FString& InMessage, iree_status_t InStatus);
+
+// IREE Environment that holds global state, e.g. containing global thread pool etc.
+class FEnvironment
+{
+public:
+	struct FConfig
+	{
+		FNNERuntimeIREEThreadingOptions ThreadingOptions;
+	};
+
+	FEnvironment();
+	~FEnvironment();
+
+	void Configure(const FConfig& InConfig);
+
+	TArray<iree_task_executor_t*> GetTaskExecutors() const;
+
+	FConfig GetConfig() const { return Config; }
+
+private:
+	bool Create() const;
+
+	FConfig Config{};
+
+	struct FInternalState;
+	mutable TUniquePtr<FInternalState> InternalState;
+
+	mutable FCriticalSection CriticalSection;
+};
+
+} // UE::NNERuntimeIREE::Private
+
+#endif // WITH_NNE_RUNTIME_IREE
\ No newline at end of file
diff --git a/Engine/Plugins/Experimental/NNERuntimeIREE/Source/NNERuntimeIREE/Private/NNERuntimeIREEModel.cpp b/Engine/Plugins/Experimental/NNERuntimeIREE/Source/NNERuntimeIREE/Private/NNERuntimeIREEModel.cpp
--- a/Engine/Plugins/Experimental/NNERuntimeIREE/Source/NNERuntimeIREE/Private/NNERuntimeIREEModel.cpp
+++ b/Engine/Plugins/Experimental/NNERuntimeIREE/Source/NNERuntimeIREE/Private/NNERuntimeIREEModel.cpp
@@ -13,11 +13,12 @@
 #include "Misc/ScopeLock.h"
 #include "NNE.h"
 #include "NNEModelData.h"
+#include "NNERuntimeIREEEnvironment.h"
 #include "NNERuntimeIREELog.h"
+#include "NNERuntimeIREESettings.h"
 #include "NNERuntimeIREETensor.h"
 #include "NNEStatus.h"
 #include "RenderGraphUtils.h"
-#include "HAL/UnrealMemory.h"
 #include "Serialization/Archive.h"
 
 #if PLATFORM_MICROSOFT
@@ -26,6 +27,7 @@
 #endif // PLATFORM_MICROSOFT
 THIRD_PARTY_INCLUDES_START
 #include "iree/hal/drivers/local_sync/sync_device.h"
+#include "iree/hal/drivers/local_task/task_device.h"
 #include "iree/hal/local/loaders/static_library_loader.h"
 #include "iree/modules/hal/types.h"
 #include "iree/runtime/call.h"
@@ -39,11 +41,12 @@ THIRD_PARTY_INCLUDES_END
 #include "Microsoft/HideMicrosoftPlatformTypes.h"
 #endif // PLATFORM_MICROSOFT
 
+LLM_DEFINE_TAG(NNERuntimeIREE_Cpu);
+
 BEGIN_SHADER_PARAMETER_STRUCT(FIREESessionRDGParameters, )
 	RDG_BUFFER_ACCESS_ARRAY(InputBuffers)
 	RDG_BUFFER_ACCESS_ARRAY(OutputBuffers)
 END_SHADER_PARAMETER_STRUCT()
-LLM_DEFINE_TAG(NNERuntimeIREE_Cpu);
 
 DECLARE_GPU_STAT_NAMED(FNNERuntimeIREERDGEnqueue, TEXT("NNERuntimeIREERdg.EnqueueRDG"));
 DECLARE_GPU_STAT_NAMED(FNNERuntimeIREERDGCall, TEXT("NNERuntimeIREERdg.EnqueueRDG.Call"));
@@ -95,67 +98,6 @@ namespace UE::NNERuntimeIREE
 			}
 		}
 
-		void PrintIREEError(const FString& InMessage, iree_status_t InStatus)
-		{
-			iree_host_size_t TrueLength = 0;
-			iree_status_format(InStatus, 0, (char*)nullptr, &TrueLength);
-			void* ErrorString = FMemory::Malloc(TrueLength + 1);
-			((char*)ErrorString)[TrueLength] = (char)0;
-			iree_status_format(InStatus, TrueLength, (char*)ErrorString, &TrueLength);
-			UE_LOG(LogNNERuntimeIREE, Error, TEXT("%s: %s"), *InMessage, *FString(StringCast<TCHAR>(static_cast<const ANSICHAR*>(ErrorString)).Get()));
-			FMemory::Free(ErrorString);
-		}
-
-		iree_status_t HostAllocatorImplAlloc(iree_allocator_command_t Command, const iree_allocator_alloc_params_t* Params, void** InoutPtr)
-		{
-			check(Params);
-			check(InoutPtr);
-			switch (Command)
-			{
-				case IREE_ALLOCATOR_COMMAND_MALLOC:
-					*InoutPtr = FMemory::Malloc(Params->byte_length, IREE_HAL_HEAP_BUFFER_ALIGNMENT);
-					break;
-				case IREE_ALLOCATOR_COMMAND_CALLOC:
-					*InoutPtr = FMemory::MallocZeroed(Params->byte_length, IREE_HAL_HEAP_BUFFER_ALIGNMENT);
-					break;
-				case IREE_ALLOCATOR_COMMAND_REALLOC:
-					*InoutPtr = FMemory::Realloc(*InoutPtr, Params->byte_length, IREE_HAL_HEAP_BUFFER_ALIGNMENT);
-					break;
-				default:
-					check(false);
-			}
-			if (!*InoutPtr)
-			{
-				return iree_make_status(IREE_STATUS_RESOURCE_EXHAUSTED, "Host allocator failed during alloc");
-			}
-			return iree_ok_status();
-		}
-
-		iree_status_t HostAllocatorImpl(void *Self, iree_allocator_command_t Command, const void* Params, void** InoutPtr)
-		{
-			LLM_SCOPE_BYTAG(NNERuntimeIREE_Cpu);
-			check(InoutPtr);
-			switch (Command)
-			{
-				case IREE_ALLOCATOR_COMMAND_MALLOC:
-				case IREE_ALLOCATOR_COMMAND_CALLOC:
-				case IREE_ALLOCATOR_COMMAND_REALLOC:
-					return HostAllocatorImplAlloc(Command, static_cast<const iree_allocator_alloc_params_t*>(Params), InoutPtr);
-				case IREE_ALLOCATOR_COMMAND_FREE:
-					FMemory::Free(*InoutPtr);
-					break;
-				default:
-					return iree_make_status(IREE_STATUS_UNIMPLEMENTED, "Host allocator command not recognised");
-			}
-			return iree_ok_status();
-		}
-
-		inline iree_allocator_t MakeHostAllocator()
-		{
-			iree_allocator_t Allocator = { NULL, HostAllocatorImpl };
-			return Allocator;
-		}
-
 		class FInstance
 		{
 		private:
@@ -322,6 +264,67 @@ namespace UE::NNERuntimeIREE
 				return true;
 			}
 
+			bool CreateTaskDevice(UE::NNERuntimeIREE::Private::FEnvironment& InEnvironment, void* InLibraryQueryFuntionPointer, iree_hal_device_t** OutDevice)
+			{
+				TArray<iree_task_executor_t*> TaskExecutors = InEnvironment.GetTaskExecutors();
+				if (TaskExecutors.IsEmpty())
+				{
+					return false;
+				}
+
+				iree_status_t Status = iree_ok_status();
+
+				iree_hal_task_device_params_t TaskDeviceParams;
+				iree_hal_task_device_params_initialize(&TaskDeviceParams);
+
+				const iree_hal_executable_library_query_fn_t Libraries[] = { (iree_hal_executable_library_query_fn_t)InLibraryQueryFuntionPointer };
+				iree_hal_executable_loader_t* LibraryLoader = NULL;
+				if (iree_status_is_ok(Status)) {
+					Status = iree_hal_static_library_loader_create(IREE_ARRAYSIZE(Libraries), Libraries, iree_hal_executable_import_provider_null(), GetHostAllocator(), &LibraryLoader);
+					if (!iree_status_is_ok(Status))
+					{
+						Private::PrintIREEError("CPU Task instance: Failed to create the static library loader", Status);
+						LibraryLoader = NULL;
+					}
+				}
+
+				iree_string_view_t Identifier = iree_make_cstring_view("local-task");
+				iree_hal_allocator_t* DeviceAllocator = NULL;
+				if (iree_status_is_ok(Status)) {
+					Status = iree_hal_allocator_create_heap(Identifier, GetHostAllocator(), GetHostAllocator(), &DeviceAllocator);
+					if (!iree_status_is_ok(Status))
+					{
+						Private::PrintIREEError("CPU Task instance: Failed to create the heap allocator", Status);
+						DeviceAllocator = NULL;
+					}
+				}
+				
+				if (iree_status_is_ok(Status)) {
+					Status = iree_hal_task_device_create(
+						Identifier,
+						&TaskDeviceParams,
+						TaskExecutors.Num(),
+						TaskExecutors.GetData(),
+						/*loader_count=*/1,
+						&LibraryLoader,
+						DeviceAllocator,
+						GetHostAllocator(),
+						OutDevice);
+					if (!iree_status_is_ok(Status))
+					{
+						Private::PrintIREEError("CPU Task instance: Failed to create the task device", Status);
+					}
+				}
+
+				if (DeviceAllocator) iree_hal_allocator_release(DeviceAllocator);
+				if (LibraryLoader) iree_hal_executable_loader_release(LibraryLoader);
+
+				bool bResult = iree_status_is_ok(Status);
+				iree_status_free(Status);
+
+				return bResult;
+			}
+
 			bool CreateSession(iree_hal_device_t* InDevice, iree_runtime_session_t** OutSession)
 			{
 				FScopeLock ScopeLock(&CriticalSection);
@@ -586,7 +589,7 @@ namespace UE::NNERuntimeIREE
 					iree_hal_device_release(Device);
 				}
 
-				static TSharedPtr<FDevice> Make(const FString& InLibraryPath, const FString& InLibraryName, const FString& LibraryQueryFunctionName)
+				static TSharedPtr<FDevice> Make(UE::NNERuntimeIREE::Private::FEnvironment& InEnvironment, const FString& InLibraryPath, const FString& InLibraryName, const FString& LibraryQueryFunctionName)
 				{
 					check(!InLibraryName.IsEmpty());
 					check(!LibraryQueryFunctionName.IsEmpty());
@@ -617,10 +620,22 @@ namespace UE::NNERuntimeIREE
 					}
 
 					iree_hal_device_t* Device = nullptr;
-					if (!Instance->CreateSyncDevice(LibraryQueryFunctionPointer, &Device))
+					if (InEnvironment.GetConfig().ThreadingOptions.bIsSingleThreaded)
 					{
-						return TSharedPtr<FDevice>();
+						if (!Instance->CreateSyncDevice(LibraryQueryFunctionPointer, &Device))
+						{
+							return TSharedPtr<FDevice>();
+						}
 					}
+					else
+					{
+						if (!Instance->CreateTaskDevice(InEnvironment, LibraryQueryFunctionPointer, &Device))
+						{
+							return TSharedPtr<FDevice>();
+						}
+					}
+
+					check(Device);
 
 					TSharedPtr<FDevice> Result = TSharedPtr<FDevice>(new FDevice(Instance.ToSharedRef(), Library.ToSharedRef(), Device));
 					Devices.Add(CombinedPathPlusFunction, Result);
@@ -1050,17 +1065,17 @@ namespace UE::NNERuntimeIREE
 
 		}
 
-		TSharedPtr<FModel> FModel::Make(const FString& InDirPath, const FString& InSharedLibraryFileName, const FString& InVmfbFileName, const FString& InLibraryQueryFunctionName, const UNNERuntimeIREEModuleMetaData& InModuleMetaData)
+		TSharedPtr<FModel> FModel::Make(UE::NNERuntimeIREE::Private::FEnvironment& InEnvironment, const FString& InDirPath, const FString& InSharedLibraryFileName, const FString& InVmfbFileName, const FString& InLibraryQueryFunctionName, const UNNERuntimeIREEModuleMetaData& InModuleMetaData)
 		{
 			check(!InSharedLibraryFileName.IsEmpty());
 			check(!InVmfbFileName.IsEmpty());
 			check(!InLibraryQueryFunctionName.IsEmpty());
 			check(!InModuleMetaData.FunctionMetaData.IsEmpty());
 
-			TSharedPtr<Private::FDevice> Device = Private::FDevice::Make(InDirPath, InSharedLibraryFileName, InLibraryQueryFunctionName);
+			TSharedPtr<Private::FDevice> Device = Private::FDevice::Make(InEnvironment, InDirPath, InSharedLibraryFileName, InLibraryQueryFunctionName);
 			if (!Device.IsValid())
 			{
-				return TSharedPtr<FModel>();
+				return {};
 			}
 
 			const FString CombinedPath = FPaths::Combine(InDirPath, InVmfbFileName);
diff --git a/Engine/Plugins/Experimental/NNERuntimeIREE/Source/NNERuntimeIREE/Private/NNERuntimeIREEModel.h b/Engine/Plugins/Experimental/NNERuntimeIREE/Source/NNERuntimeIREE/Private/NNERuntimeIREEModel.h
--- a/Engine/Plugins/Experimental/NNERuntimeIREE/Source/NNERuntimeIREE/Private/NNERuntimeIREEModel.h
+++ b/Engine/Plugins/Experimental/NNERuntimeIREE/Source/NNERuntimeIREE/Private/NNERuntimeIREEModel.h
@@ -16,6 +16,7 @@ namespace UE::NNERuntimeIREE
 {
 	namespace Private
 	{
+		class FEnvironment;
 		class FModule;
 	} // Private
 
@@ -58,7 +59,7 @@ namespace UE::NNERuntimeIREE
 			FModel(TSharedRef<UE::NNERuntimeIREE::CPU::Private::FDevice> InDevice, TSharedRef<UE::NNERuntimeIREE::Private::FModule> InModule);
 
 		public:
-			static TSharedPtr<FModel> Make(const FString& InDirPath, const FString& InSharedLibraryFileName, const FString& InVmfbFileName, const FString& InLibraryQueryFunctionName, const UNNERuntimeIREEModuleMetaData& InModuleMetaData);
+			static TSharedPtr<FModel> Make(UE::NNERuntimeIREE::Private::FEnvironment& InEnvironment, const FString& InDirPath, const FString& InSharedLibraryFileName, const FString& InVmfbFileName, const FString& InLibraryQueryFunctionName, const UNNERuntimeIREEModuleMetaData& InModuleMetaData);
 
 		public:
 			//~ Begin IModelCPU Interface
diff --git a/Engine/Plugins/Experimental/NNERuntimeIREE/Source/NNERuntimeIREE/Private/NNERuntimeIREEModule.cpp b/Engine/Plugins/Experimental/NNERuntimeIREE/Source/NNERuntimeIREE/Private/NNERuntimeIREEModule.cpp
--- a/Engine/Plugins/Experimental/NNERuntimeIREE/Source/NNERuntimeIREE/Private/NNERuntimeIREEModule.cpp
+++ b/Engine/Plugins/Experimental/NNERuntimeIREE/Source/NNERuntimeIREE/Private/NNERuntimeIREEModule.cpp
@@ -8,6 +8,8 @@
 
 #include "CoreMinimal.h"
 #include "NNE.h"
+#include "NNERuntimeIREEEnvironment.h"
+#include "NNERuntimeIREESettings.h"
 
 #if WITH_EDITOR
 #include "HAL/FileManager.h"
@@ -25,9 +27,9 @@
 
 DEFINE_LOG_CATEGORY(LogNNERuntimeIREE);
 
-#if WITH_EDITOR
 namespace UE::NNERuntimeIREE::Private
 {
+#if WITH_EDITOR
 	namespace DllHelper
 	{
 		void* GetDllHandle(const FString& Path)
@@ -50,14 +52,33 @@ namespace UE::NNERuntimeIREE::Private
 			return Result;
 		}
 	} // namespace DllHelper
-} // namespace UE::NNERuntimeIREE::Private
 #endif // WITH_EDITOR
 
+#ifdef WITH_NNE_RUNTIME_IREE
+	namespace EnvironmentHelper
+	{
+		void CreateEnvironmentFromSettings(const UNNERuntimeIREESettings* Settings, FEnvironment& Environment)
+		{
+#if WITH_EDITOR
+			const FNNERuntimeIREEThreadingOptions& ThreadingOptions = Settings->EditThreadingOptions;
+#else
+			const FNNERuntimeIREEThreadingOptions& ThreadingOptions = Settings->GameThreadingOptions;
+#endif
+
+			UE::NNERuntimeIREE::Private::FEnvironment::FConfig Config{};
+			Config.ThreadingOptions = ThreadingOptions;
+
+			Environment.Configure(Config);
+		}
+	} // EnvironmentHelper
+#endif // WITH_NNE_RUNTIME_IREE
+} // namespace UE::NNERuntimeIREE::Private
+
 void FNNERuntimeIREEModule::StartupModule()
 {
-#if WITH_EDITOR
 	using namespace UE::NNERuntimeIREE::Private;
 
+#if WITH_EDITOR
 	const FString PluginDir = IPluginManager::Get().FindPlugin("NNERuntimeIREE")->GetBaseDir();
 	const FString MlirToolsSharedLibPath = FPaths::Combine(PluginDir, TEXT(PREPROCESSOR_TO_STRING(NNEMLIRTOOLS_SHAREDLIB_PATH)));
 
@@ -86,9 +107,15 @@ void FNNERuntimeIREEModule::StartupModule()
 #endif // WITH_EDITOR
 
 #ifdef WITH_NNE_RUNTIME_IREE
+	Environment = MakeShared<FEnvironment>();
+
+	EnvironmentHelper::CreateEnvironmentFromSettings(GetDefault<UNNERuntimeIREESettings>(), *Environment);
+
 	NNERuntimeIREECpu = NewObject<UNNERuntimeIREECpu>();
 	if (NNERuntimeIREECpu.IsValid())
 	{
+		NNERuntimeIREECpu->Init(Environment.ToSharedRef());
+
 		NNERuntimeIREECpu->AddToRoot();
 		UE::NNE::RegisterRuntime(NNERuntimeIREECpu.Get());
 	}
@@ -170,6 +197,8 @@ void FNNERuntimeIREEModule::ShutdownModule()
 		NNERuntimeIREERdg->RemoveFromRoot();
 		NNERuntimeIREERdg.Reset();
 	}
+
+	Environment.Reset();
 #endif // WITH_NNE_RUNTIME_IREE
 }
 	
diff --git a/Engine/Plugins/Experimental/NNERuntimeIREE/Source/NNERuntimeIREE/Private/NNERuntimeIREEModule.h b/Engine/Plugins/Experimental/NNERuntimeIREE/Source/NNERuntimeIREE/Private/NNERuntimeIREEModule.h
--- a/Engine/Plugins/Experimental/NNERuntimeIREE/Source/NNERuntimeIREE/Private/NNERuntimeIREEModule.h
+++ b/Engine/Plugins/Experimental/NNERuntimeIREE/Source/NNERuntimeIREE/Private/NNERuntimeIREEModule.h
@@ -7,6 +7,12 @@
 #ifdef WITH_NNE_RUNTIME_IREE
 #include "NNERuntimeIREE.h"
 #include "UObject/WeakObjectPtrTemplates.h"
+
+namespace UE::NNERuntimeORT::Private
+{
+	class FEnvironment;
+}
+
 #endif // WITH_NNE_RUNTIME_IREE
 
 class FNNERuntimeIREEModule : public IModuleInterface
@@ -23,5 +29,6 @@ private:
 	TWeakObjectPtr<UNNERuntimeIREECuda> NNERuntimeIREECuda;
 	TWeakObjectPtr<UNNERuntimeIREEVulkan> NNERuntimeIREEVulkan;
 	TWeakObjectPtr<UNNERuntimeIREERdg> NNERuntimeIREERdg;
+	TSharedPtr<UE::NNERuntimeIREE::Private::FEnvironment> Environment;
 #endif // WITH_NNE_RUNTIME_IREE
 };
diff --git a/Engine/Plugins/Experimental/NNERuntimeIREE/Source/NNERuntimeIREE/Private/NNERuntimeIREESettings.cpp b/Engine/Plugins/Experimental/NNERuntimeIREE/Source/NNERuntimeIREE/Private/NNERuntimeIREESettings.cpp
new file mode 100644
--- /dev/null
+++ b/Engine/Plugins/Experimental/NNERuntimeIREE/Source/NNERuntimeIREE/Private/NNERuntimeIREESettings.cpp
@@ -0,0 +1,21 @@
+// Copyright Epic Games, Inc. All Rights Reserved.
+
+#include "NNERuntimeIREESettings.h"
+
+UNNERuntimeIREESettings::UNNERuntimeIREESettings(const FObjectInitializer& ObjectInitlaizer)
+	: Super(ObjectInitlaizer)
+{
+
+}
+
+FName UNNERuntimeIREESettings::GetCategoryName() const
+{
+	return TEXT("Plugins");
+}
+
+#if WITH_EDITOR
+FText UNNERuntimeIREESettings::GetSectionText() const
+{
+	return NSLOCTEXT("NNERuntimeIREEPlugin", "NNERuntimeIREESettingsSection", "NNERuntimeIREE");
+}
+#endif
\ No newline at end of file
diff --git a/Engine/Plugins/Experimental/NNERuntimeIREE/Source/NNERuntimeIREE/Private/NNERuntimeIREESettings.h b/Engine/Plugins/Experimental/NNERuntimeIREE/Source/NNERuntimeIREE/Private/NNERuntimeIREESettings.h
new file mode 100644
--- /dev/null
+++ b/Engine/Plugins/Experimental/NNERuntimeIREE/Source/NNERuntimeIREE/Private/NNERuntimeIREESettings.h
@@ -0,0 +1,101 @@
+// Copyright Epic Games, Inc. All Rights Reserved.
+
+#pragma once
+
+#include "Engine/DeveloperSettings.h"
+
+#include "NNERuntimeIREESettings.generated.h"
+
+/** Specifies the thread affinity group type */
+UENUM()
+enum class ENNERuntimeIREEThreadingAffinityGroupSpecifierType : uint8
+{
+	Index 		UMETA(DisplayName = "Specify group by index"),
+	Current		UMETA(DisplayName = "Use same group as calling thread"),
+	All			UMETA(DisplayName = "Use all groups"),
+	Any			UMETA(DisplayName = "Use any group")
+};
+
+/** Specifies the processor affinity for a particular thread given by processor group the thread should be assigned to (aka NUMA node, cluster etc., depending on the platform) and the processor it should be scheduled on */
+USTRUCT()
+struct FNNERuntimeIREEThreadingAffinity
+{
+	GENERATED_BODY()
+
+	/** How is the group specified? */
+	UPROPERTY(Config, EditAnywhere, Category = "IREE Runtime", meta = (
+		DisplayName = "Group specifier",
+		ToolTip = "Specify the group (type) the thread should be assigned to."))
+	ENNERuntimeIREEThreadingAffinityGroupSpecifierType GroupSpecifierType = ENNERuntimeIREEThreadingAffinityGroupSpecifierType::Any;
+
+	/** Group index, only used if node specifier type is 'Index' */
+	UPROPERTY(Config, EditAnywhere, Category = "IREE Runtime", meta = (
+		DisplayName = "Group index",
+		ToolTip = "Specify with an index the group the thread should be assigned to. Only used if group specifier type is 'Index'."))
+	int32 GroupIndex = -1;
+
+	/** Logical core index */
+	UPROPERTY(Config, EditAnywhere, Category = "IREE Runtime", meta = (
+		DisplayName = "Logical core index",
+		ToolTip = "Specify with an index the logical core the thread should be scheduled on. Set to -1 to let the thread be scheduled on any core."))
+	int32 CoreIndex = -1;
+};
+
+/** Specifies the task topology used for multi-threading */
+USTRUCT()
+struct FNNERuntimeIREETaskTopology
+{
+	GENERATED_BODY()
+
+	/** Task topology groups */
+	UPROPERTY(Config, EditAnywhere, Category = "IREE Runtime", meta = (
+		DisplayName = "Task topology groups",
+		ToolTip = "Specify the task system topology with zero or more groups. If empty, the runtime will try to use all physical cores."))
+	TArray<FNNERuntimeIREEThreadingAffinity> TaskTopologyGroups;
+};
+
+/** Threading options to configure single/multi-threading, including task topology for multi-threading case */
+USTRUCT()
+struct FNNERuntimeIREEThreadingOptions
+{
+	GENERATED_BODY()
+
+	/** Run single threaded? */
+	UPROPERTY(Config, EditAnywhere, Category = "IREE Runtime", meta = (
+		DisplayName = "Is single threaded",
+		ToolTip = "Should the runtime run single threaded (might result in better performance with very small models)."))
+	bool bIsSingleThreaded = true;
+
+	/** Task topology used for multi-threading */
+	UPROPERTY(Config, EditAnywhere, Category = "IREE Runtime", meta = (
+		DisplayName = "Task topology",
+		ToolTip = "For multi-threading you can specify the task system topology."))
+	FNNERuntimeIREETaskTopology TaskTopology;
+};
+
+/** Settings used to configure NNERuntimeIREE */
+UCLASS(config = Engine, defaultconfig, meta = (DisplayName = "NNERuntimeIREE"))
+class UNNERuntimeIREESettings : public UDeveloperSettings
+{
+	GENERATED_BODY()
+
+	UNNERuntimeIREESettings(const FObjectInitializer& ObjectInitlaizer);
+
+public:
+
+	/** Threading options in Editor targets */
+	UPROPERTY(Config, EditAnywhere, Category = "IREE Runtime", meta = (DisplayName = "Editor threading options", ToolTip = "Threading options in Editor targets", ConfigRestartRequired = true))
+	FNNERuntimeIREEThreadingOptions EditThreadingOptions{};
+
+	/** Threading options in Non-Editor (Game, Program, ...) targets */
+	UPROPERTY(Config, EditAnywhere, Category = "IREE Runtime", meta = (DisplayName = "Game threading options", ToolTip = "Threading options in Non-Editor (Game, Program, ...) targets", ConfigRestartRequired = true))
+	FNNERuntimeIREEThreadingOptions GameThreadingOptions{};
+
+	// Begin UDeveloperSettings Interface
+	NNERUNTIMEIREE_API virtual FName GetCategoryName() const override;
+#if WITH_EDITOR
+	NNERUNTIMEIREE_API virtual FText GetSectionText() const override;
+	// END UDeveloperSettings Interface
+#endif
+
+};
