[NNE] NNERuntimeIREE fixed IntProperty not initialized properly warning.

#rnx

diff --git a/Engine/Plugins/Experimental/NNERuntimeIREE/Source/NNERuntimeIREE/Private/NNERuntimeIREEModelData.h b/Engine/Plugins/Experimental/NNERuntimeIREE/Source/NNERuntimeIREE/Private/NNERuntimeIREEModelData.h
--- a/Engine/Plugins/Experimental/NNERuntimeIREE/Source/NNERuntimeIREE/Private/NNERuntimeIREEModelData.h
+++ b/Engine/Plugins/Experimental/NNERuntimeIREE/Source/NNERuntimeIREE/Private/NNERuntimeIREEModelData.h
@@ -71,7 +71,7 @@ struct FNNERuntimeIREEModelDataHeaderRDG
 	 * Current version of this IREE model data.
 	 */
 	UPROPERTY()
-	int32 Version;
+	int32 Version = 0;
 
 	/**
 	 * A Guid that uniquely identifies the model.
