[NNE] NNERuntimeIREE added IREE build scripts.

Problem:
IREE build scripts were in a separate Git and not available to customers. Multiple customes asked for access to be able to create thier own custom builds.

Solution:
Moved all build scripts and code of our own custom external tools into NNERuntimeIREE Plugin folder tree. Platform specific build scripts are in their respective Plugin-Extension folder tree. They are automatically discovered by the main build script (Windows). The build artifacts are copied to their target locations and a Perfoce CL, including the modifications, is created.

#tests Win64, Linux and Mac,Editor, tested IREE 3.5.0 upgrade.

#rnx

diff --git a/Engine/Plugins/Experimental/NNERuntimeIREE/Source/ExternalBuild/CMake/iree_runtime.cmake b/Engine/Plugins/Experimental/NNERuntimeIREE/Source/ExternalBuild/CMake/iree_runtime.cmake
new file mode 100644
--- /dev/null
+++ b/Engine/Plugins/Experimental/NNERuntimeIREE/Source/ExternalBuild/CMake/iree_runtime.cmake
@@ -0,0 +1,35 @@
+cmake_minimum_required(VERSION 3.21...3.24)
+
+if (NOT UE_IREE_BUILD_ROOT)
+	message(FATAL_ERROR "Error: UE_IREE_BUILD_ROOT not set")
+endif()
+
+if (NOT UE_IREE_ADDITIONAL_SOURCES)
+	set(UE_IREE_ADDITIONAL_SOURCES "")
+endif()
+
+# Set device standalone to make sure some functions are included to avoid mismatch on target systems (e.g. __truncsfhf2 and __extendhfsf2)
+add_compile_definitions(IREE_DEVICE_STANDALONE)
+
+add_subdirectory("${UE_IREE_BUILD_ROOT}/iree" "iree" EXCLUDE_FROM_ALL)
+
+# Get all targets required by the runtime
+set(_RUNTIME_ROOT "iree::runtime")
+set(_RUNTIME_OBJECTS "$<REMOVE_DUPLICATES:$<GENEX_EVAL:$<TARGET_PROPERTY:${_RUNTIME_ROOT},INTERFACE_IREE_TRANSITIVE_OBJECTS>>>")
+set(_RUNTIME_LIBS "$<REMOVE_DUPLICATES:$<GENEX_EVAL:$<TARGET_PROPERTY:${_RUNTIME_ROOT},INTERFACE_IREE_TRANSITIVE_OBJECT_LIBS>>>")
+
+set(_LOADER_ROOT "iree::hal::local::loaders::static_library_loader")
+set(_LOADER_OBJECTS "$<REMOVE_DUPLICATES:$<GENEX_EVAL:$<TARGET_PROPERTY:${_LOADER_ROOT},INTERFACE_IREE_TRANSITIVE_OBJECTS>>>")
+set(_LOADER_LIBS "$<REMOVE_DUPLICATES:$<GENEX_EVAL:$<TARGET_PROPERTY:${_LOADER_ROOT},INTERFACE_IREE_TRANSITIVE_OBJECT_LIBS>>>")
+
+set(_ELF_ROOT "iree::hal::local::elf::arch")
+set(_ELF_OBJECTS "$<REMOVE_DUPLICATES:$<GENEX_EVAL:$<TARGET_PROPERTY:${_ELF_ROOT},INTERFACE_IREE_TRANSITIVE_OBJECTS>>>")
+set(_ELF_LIBS "$<REMOVE_DUPLICATES:$<GENEX_EVAL:$<TARGET_PROPERTY:${_ELF_ROOT},INTERFACE_IREE_TRANSITIVE_OBJECT_LIBS>>>")
+
+set(_DEVICE_ROOT "iree::builtins::device::device")
+set(_DEVICE_OBJECTS "$<REMOVE_DUPLICATES:$<GENEX_EVAL:$<TARGET_PROPERTY:${_DEVICE_ROOT},INTERFACE_IREE_TRANSITIVE_OBJECTS>>>")
+set(_DEVICE_LIBS "$<REMOVE_DUPLICATES:$<GENEX_EVAL:$<TARGET_PROPERTY:${_DEVICE_ROOT},INTERFACE_IREE_TRANSITIVE_OBJECT_LIBS>>>")
+
+# Define the library
+add_library(ireert ${_RUNTIME_OBJECTS} ${_LOADER_OBJECTS} ${_ELF_OBJECTS} ${_DEVICE_OBJECTS} ${UE_IREE_ADDITIONAL_SOURCES})
+target_link_libraries(ireert PUBLIC ${_RUNTIME_LIBS} ${_LOADER_LIBS} ${_ELF_LIBS} ${_DEVICE_LIBS})
\ No newline at end of file
diff --git a/Engine/Plugins/Experimental/NNERuntimeIREE/Source/ExternalBuild/Include/Clang/alloca.h b/Engine/Plugins/Experimental/NNERuntimeIREE/Source/ExternalBuild/Include/Clang/alloca.h
new file mode 100644
--- /dev/null
+++ b/Engine/Plugins/Experimental/NNERuntimeIREE/Source/ExternalBuild/Include/Clang/alloca.h
@@ -0,0 +1,15 @@
+// Copyright Epic Games, Inc. All Rights Reserved.
+
+#pragma once
+
+#ifdef __cplusplus
+extern "C" {
+#endif // __cplusplus
+
+void *alloca(size_t);
+
+#define alloca __builtin_alloca
+
+#ifdef __cplusplus
+}
+#endif // __cplusplus
\ No newline at end of file
diff --git a/Engine/Plugins/Experimental/NNERuntimeIREE/Source/ExternalBuild/README.md b/Engine/Plugins/Experimental/NNERuntimeIREE/Source/ExternalBuild/README.md
new file mode 100644
--- /dev/null
+++ b/Engine/Plugins/Experimental/NNERuntimeIREE/Source/ExternalBuild/README.md
@@ -0,0 +1,64 @@
+# NNERuntimeIREE build external dependencies
+Scripts to clone and build the IREE compiler as well as the runtime libraries. Additionally will build the shared library NNEMlirTools that is used to parse and inspect MLIR. The library is based on the same IREE clone to enable maximum compability.
+
+## Windows
+
+### Setup & Requirements
+
+- Install Visual Studio 2022
+- Start the Visual Studio Installer, click 'Modify' and change to the 'Individual Components' tab
+- Check which toolchains are supported in UE (written in Engine\Config\Windows\Windows_SDK.json in PreferredVisualCppVersions, e.g "14.38.33130-14.38.99999" // VS2022 17.8.x for UE 5.5 and write down the lowest supported MSVC-VS versions, futher down in this doc it will be called *lowest MSVC-VS XX.XX-XX.X* or *lowest VS XX.X*. For UE 5.5 it is thus 14.38-17.8.
+- Check which SDK is supported in UE (also written in Engine\Config\Windows\Windows_SDK.json but under "MinVersion") it will be called *lowest supported WINSDK ver* in this doc.
+- Make sure you install the following toolchains (and only the following, there were some issues when mixing different toolchains)
+  - C++ Modules for v143 build tools (x64/x86 - experimental)
+  - C++/CLI support for v143 ( *lowest MSVC-VS XX.XX-XX.X*)
+  - MSBuild
+  - MSBuild support for LLVM (clang-cl) toolset
+  - MSVC v143 - VS 2022 C++ ARM build tools (*lowest MSVC-VS XX.XX-XX.X*)
+  - MSVC v143 - VS 2022 C++ ARM Spectre-mitigated libs (*lowest MSVC-VS XX.XX-XX.X*)
+  - MSVC v143 - VS 2022 C++ ARM64/ARM64EC build tools (*lowest MSVC-VS XX.XX-XX.X*)
+  - MSVC v143 - VS 2022 C++ ARM64/ARM64EC Spectre-mitigated libs (*lowest MSVC-VS XX.XX-XX.X*)
+  - MSVC v143 - VS 2022 C++ x64/x86 build tools (*lowest MSVC-VS XX.XX-XX.X*)
+  - MSVC v143 - VS 2022 C++ x64/x86 Spectre-mitigated libs (*lowest MSVC-VS XX.XX-XX.X*)
+  - MSVC v143 C++ X64/x86 build tools (Latest) <-- required to get 'vcvarsall.bat'
+- Also install
+  - All six C++ ATL for (*lowest VS XX.X*) v143 build tools
+  - All six C++ MFC for (*lowest VS XX.X*) v143 build tools
+  - Windows SDK (*lowest supported WINSDK ver*) witch is used for the runtime lib
+  - Windows 11 SDK (10.0.22621.0) witch is used for the iree compiler
+- Install Cmake and also ninja as described on the IREE webpage
+  - In addition, set the system environment variable NINJA_PATH to your executable (e.g. D:/ninja/ninja.exe)
+  - Set the vcvarsall.bat path to an environment variable called VCVARSALL_PATH (e.g.C:\Program Files\Microsoft Visual Studio\2022\Professional\VC\Auxiliary\Build\vcvarsall.bat for UE 5.5)
+- Install all console SDKs and make sure the corresponding environment variables are set
+- Install git command tools
+
+## Building & Install
+
+The build scripts can be found in the subdirectory 'Scripts'. For every build platform there is a subdirectory:
+- `Build*.*`: Main build script.
+- `Install*.*`: Script that copies artifacts to their specific location within the NNERuntimeIREE plugin and finally creates a CL in Perforce with the changes included.
+- `BuildAndInstall*.*`: Convencience script that calls `Build*.*` and `Install*.*`.
+
+For customization you can pass the intermediate build directory path to the scripts. On Windows you also can provide the additional platforms the IREE runtime should be built for (only if supported). Or by specifing 'All', the script will look by itself for all platforms that are supported and build the runtime for them.
+
+*Windows*
+
+`Script.bat [-dir WorkingDir] [-platform "Platform1 Platform2 ..."]`
+
+*Linux & Mac*
+
+`Script.sh [WorkingDir]`
+
+**Important note**
+
+Since source paths can be shown in IREE error messages, be adviced for production builds to carefully setup the folder structure. A common approach we choose is to build in `github.com` which is located at the root of the preferred volume or disk (e.g. `E:\github.com` on Windows, `/github.com` on Linux, or `/Volumes/MyVolume/github.com` on Mac).
+
+## Linux
+
+- Install cmake and ninja-build
+- Ensure v23_clang-18.1.0-rockylinux8.tar.gz is still the compiler used for UE (was at 5.5)
+  * To do so you need to get UE building and launch any sample project. Then navigate to /Engine/Build/BatchFiles/Linux/ and execute ./GenerateProjectFiles.sh -cmakefile -game "/path/to/the/sample/project.uproject" you will see in the logs what is used compiler. If different you need to adapt COMPILER_NAME in BuildLinux.sh
+- The llvm project is built with clang-16, so make sure this compiler is installed too
+- clang-16 requires libstdc++-12, it can be installed from the default repositories: sudo apt install libstdc++-12-dev.
+ 
+Note: atm llvm cannot be compiled using the clang that UE is using. Thus a different compiler is used for this.
diff --git a/Engine/Plugins/Experimental/NNERuntimeIREE/Source/ExternalBuild/Scripts/Linux/BuildAndInstallLinux.sh b/Engine/Plugins/Experimental/NNERuntimeIREE/Source/ExternalBuild/Scripts/Linux/BuildAndInstallLinux.sh
new file mode 100755
--- /dev/null
+++ b/Engine/Plugins/Experimental/NNERuntimeIREE/Source/ExternalBuild/Scripts/Linux/BuildAndInstallLinux.sh
@@ -0,0 +1,39 @@
+#! /bin/bash
+
+#
+# Script that calls build and install scripts together with the same arguments.
+#
+# If any intermediate step fails, the script (should) exit.
+#
+
+set -Eeuo pipefail
+
+trap 'echo; echo "Script failed. Exiting" >&2; exit 1' ERR
+
+SCRIPT_DIR="$(cd -- "$(dirname -- "${BASH_SOURCE[0]}")" &>/dev/null && pwd)"
+
+# For production be aware that the build path is part of the IREE error messages
+# So choose it with care
+DEFAULT_WORKING_DIR="$SCRIPT_DIR/../../Build"
+
+if [ -z "${1-}" ]; then
+	mkdir -p "$DEFAULT_WORKING_DIR"
+	WORKING_DIR="$DEFAULT_WORKING_DIR"
+else
+	WORKING_DIR="$1"
+	if [ ! -d "$WORKING_DIR" ]; then
+		echo "Error: '$WORKING_DIR' does not exist." >&2
+		exit 1
+	fi
+fi
+
+WORKING_DIR="$(cd "$WORKING_DIR" && pwd)"
+
+echo
+echo "Working dir: \"$WORKING_DIR\""
+echo
+
+"${SCRIPT_DIR}/BuildLinux.sh" "$WORKING_DIR"
+"${SCRIPT_DIR}/InstallLinux.sh" "$WORKING_DIR"
+
+exit 0
\ No newline at end of file
diff --git a/Engine/Plugins/Experimental/NNERuntimeIREE/Source/ExternalBuild/Scripts/Linux/BuildLinux.sh b/Engine/Plugins/Experimental/NNERuntimeIREE/Source/ExternalBuild/Scripts/Linux/BuildLinux.sh
new file mode 100755
--- /dev/null
+++ b/Engine/Plugins/Experimental/NNERuntimeIREE/Source/ExternalBuild/Scripts/Linux/BuildLinux.sh
@@ -0,0 +1,280 @@
+#! /bin/bash
+
+#
+# This script builds all the required targets for the NNERuntimeIREE plugin.
+#
+# If an intermediate step fails, the script exits.
+#
+
+set -euo pipefail
+
+if [ -z "$1" ]; then
+	echo "Usage: [Script] [Working dir]"
+	exit 1
+fi
+
+WORKING_DIR=$1
+
+BUILD_DIR=$WORKING_DIR/iree-org
+BUILD_SCRIPT_DIR="$(dirname -- "$(realpath -- "$0")")/"
+UE_SOURCE_DIR=$BUILD_SCRIPT_DIR/../../Source
+UE_PATCHES_DIR=$BUILD_SCRIPT_DIR/../../Patch
+UE_CMAKE_DIR=$BUILD_SCRIPT_DIR/../../CMake
+UE_PLUGIN_ROOT=$BUILD_SCRIPT_DIR/../../../..
+
+echo Plugin root dir: "$UE_PLUGIN_ROOT"
+echo Build dir: "$BUILD_DIR"
+echo Source dir: "$UE_SOURCE_DIR"
+
+IREE_GIT_REPOSITORY=[REDACTED_PROFILE_URL]/iree.git
+IREE_GIT_COMMIT=v3.5.0
+IREE_THIRD_PARTY_LIBRARIES=("flatcc" "llvm-project" "stablehlo" "torch-mlir" "benchmark" "spirv_cross")
+IREE_COMPILER_VERSION_STRING="IREE-for-UE"
+
+BUILD_TYPE=Release
+
+# Unreal Shader compiler plugin
+IREE_CMAKE_PLUGIN_PATH=../../Iree/Compiler/Plugins/Target/UnrealShader
+
+SPIRV_CROSS_DIR=$BUILD_DIR/iree/third_party/spirv_cross
+
+# Change this when compiling on arch64 for arch64
+TARGET=x86_64-unknown-linux-gnu #TARGET=aarch64-unknown-linux-gnueabi
+
+COMPILER_NAME=v26_clang-20.1.8-rockylinux8
+C_COMPILER=$BUILD_DIR/Clang/$COMPILER_NAME/$TARGET/bin/clang
+CXX_COMPILER=$BUILD_DIR/Clang/$COMPILER_NAME/$TARGET/bin/clang++
+COMPILER_FLAGS=--target=$TARGET
+
+# We could not build the llvm project with the UE clang so far
+HOST_BINARIES_C_COMPILER=clang-20
+HOST_BINARIES_CXX_COMPILER=clang++-20
+
+if [ ! -d $BUILD_DIR ] ; then
+	mkdir $BUILD_DIR
+fi
+
+echo =========================================
+echo ========= INSTALLING PATCHELF ===========
+echo =========================================
+
+sudo apt-get install patchelf
+
+echo =========================================
+echo =========== INSTALLING CLANG ============
+echo =========================================
+
+if [ ! -f "$BUILD_DIR/Clang/native-linux-$COMPILER_NAME.tar.gz" ]; then
+	cd $BUILD_DIR
+
+	if [ ! -d Clang ] ; then
+		mkdir Clang
+	fi
+
+	cd Clang
+
+	if [ ! -f native-linux-$COMPILER_NAME.tar.gz ] ; then
+		wget https://cdn.unrealengine.com/Toolchain_Linux/native-linux-$COMPILER_NAME.tar.gz
+	fi
+	if [ ! -d v$COMPILER_NAME ] ; then
+		tar -xzf native-linux-$COMPILER_NAME.tar.gz
+	fi
+
+	if [ ! -f llvm.sh ] ; then
+		wget https://apt.llvm.org/llvm.sh
+		chmod 777 llvm.sh
+	fi
+	sudo ./llvm.sh 20 all
+
+	cd ../..
+	echo
+	echo Installing Clang: Done
+	echo
+else
+	echo
+	echo Installing Clang: Skipped
+	echo
+fi
+
+
+echo =========================================
+echo ============= Copy Source ===============
+echo =========================================
+
+mkdir -p "$WORKING_DIR/iree/compiler"
+cp -a "$UE_SOURCE_DIR/Iree/Compiler/." "$WORKING_DIR/iree/compiler/" > /dev/null
+
+echo
+echo Copied IREE compiler plugin.
+echo
+
+
+echo =========================================
+echo ============= Cloning IREE ==============
+echo =========================================
+
+if [ ! -d "$BUILD_DIR/iree" ]; then
+	cd $BUILD_DIR
+	git clone -n $IREE_GIT_REPOSITORY
+	cd iree
+	echo Using IREE git commit $IREE_GIT_COMMIT
+	git checkout $IREE_GIT_COMMIT
+	cd third_party
+	for D in */; do
+		if [ ! -f "$D/CMakeLists.txt" ]; then
+			echo -n "" > "$D/CMakeLists.txt"
+		fi
+	done
+	for L in ${IREE_THIRD_PARTY_LIBRARIES[@]}; do
+		if [ -d $L ]; then
+			rm "$L/CMakeLists.txt"
+			git submodule update --init -- $L
+			if [ -d "$L/third-party" ]; then
+				rm -R -f "$L/third-party"
+			fi
+		fi
+	done
+
+	cd ../../..
+
+	echo Check for $SPIRV_CROSS_DIR
+	if [ -d "$SPIRV_CROSS_DIR" ]; then
+		echo Apply git patch to spirv_cross
+		cd $SPIRV_CROSS_DIR
+		git apply $UE_PATCHES_DIR/spirv_cross.patch
+		cd ../../../..
+		echo Done.
+	fi
+
+	echo
+	echo Cloning IREE: Done
+	echo
+else
+	echo
+	echo Cloning IREE: Skipped
+	echo
+fi
+
+echo =========================================
+echo ======== Building IREE Compiler =========
+echo =========================================
+
+if [ ! -d "$BUILD_DIR/iree-compiler" ]; then
+	cmake -G Ninja -B $BUILD_DIR/iree-compiler $BUILD_DIR/iree \
+		-DCMAKE_BUILD_TYPE=$BUILD_TYPE \
+		-DIREE_CMAKE_PLUGIN_PATHS=$IREE_CMAKE_PLUGIN_PATH \
+		-DIREE_ENABLE_CPUINFO=OFF \
+		-DIREE_BUILD_TESTS=OFF \
+		-DIREE_BUILD_SAMPLES=OFF \
+		-DIREE_BUILD_BINDINGS_TFLITE=OFF \
+		-DIREE_BUILD_BINDINGS_TFLITE_JAVA=OFF \
+		-DIREE_BUILD_ALL_CHECK_TEST_MODULES=OFF \
+		-DIREE_HAL_DRIVER_DEFAULTS=OFF \
+		-DIREE_HAL_DRIVER_LOCAL_SYNC=ON \
+		-DIREE_HAL_DRIVER_LOCAL_TASK=ON \
+		-DIREE_TARGET_BACKEND_DEFAULTS=OFF \
+		-DIREE_TARGET_BACKEND_LLVM_CPU=ON \
+		-DIREE_ERROR_ON_MISSING_SUBMODULES=OFF \
+		-DCMAKE_C_COMPILER=$HOST_BINARIES_C_COMPILER \
+		-DCMAKE_C_FLAGS=$COMPILER_FLAGS \
+		-DCMAKE_CXX_COMPILER=$HOST_BINARIES_CXX_COMPILER \
+		-DCMAKE_CXX_FLAGS=$COMPILER_FLAGS \
+		-DLLVM_ENABLE_TERMINFO=NO \
+		-DIREE_EMBEDDED_RELEASE_INFO=ON \
+		-DIREE_RELEASE_VERSION=$IREE_COMPILER_VERSION_STRING \
+		-DIREE_RELEASE_REVISION=$IREE_GIT_COMMIT
+		
+	cmake --build $BUILD_DIR/iree-compiler
+	
+    echo
+	echo Building Compiler: Done
+	echo
+else
+	echo
+	echo Building Compiler: Incremental build... please do not use for production!
+	echo
+
+	cmake --build $BUILD_DIR/iree-compiler
+fi
+
+
+echo =========================================
+echo ======== Building NNEMlirTools ==========
+echo =========================================
+
+if [ ! -d "$BUILD_DIR/NNEMlirTools" ]; then
+	cmake -G Ninja -B $BUILD_DIR/NNEMlirTools $UE_SOURCE_DIR/NNEMlirTools \
+		-DCMAKE_BUILD_TYPE=$BUILD_TYPE \
+		-DUE_IREE_BUILD_ROOT=$BUILD_DIR \
+		-DIREE_ENABLE_CPUINFO=OFF \
+		-DIREE_BUILD_TESTS=OFF \
+		-DIREE_BUILD_SAMPLES=OFF \
+		-DIREE_BUILD_BINDINGS_TFLITE=OFF \
+		-DIREE_BUILD_BINDINGS_TFLITE_JAVA=OFF \
+		-DIREE_BUILD_ALL_CHECK_TEST_MODULES=OFF \
+		-DIREE_HAL_DRIVER_DEFAULTS=OFF \
+		-DIREE_HAL_DRIVER_LOCAL_SYNC=OFF \
+		-DIREE_HAL_DRIVER_LOCAL_TASK=OFF \
+		-DIREE_TARGET_BACKEND_DEFAULTS=OFF \
+		-DIREE_TARGET_BACKEND_LLVM_CPU=OFF \
+		-DIREE_ERROR_ON_MISSING_SUBMODULES=OFF \
+		-DCMAKE_C_COMPILER=$HOST_BINARIES_C_COMPILER \
+		-DCMAKE_C_FLAGS=$COMPILER_FLAGS \
+		-DCMAKE_CXX_COMPILER=$HOST_BINARIES_CXX_COMPILER \
+		-DCMAKE_CXX_FLAGS=$COMPILER_FLAGS \
+		-DLLVM_ENABLE_TERMINFO=NO 
+		
+	cmake --build $BUILD_DIR/NNEMlirTools
+	
+    echo
+	echo Building NNEMlirTools: Done
+	echo
+else
+	echo
+	echo Building NNEMlirTools: Incremental build... please do not use for production!
+	echo
+
+	cmake --build $BUILD_DIR/NNEMlirTools
+fi
+
+
+echo =========================================
+echo ========= Building Linux Runtime ==========
+echo =========================================
+
+if [ ! -d "$BUILD_DIR/iree-runtime-linux" ]; then
+	cmake -G Ninja -B $BUILD_DIR/iree-runtime-linux $UE_SOURCE_DIR/Iree/Runtime \
+		-DCMAKE_BUILD_TYPE=$BUILD_TYPE \
+		-DCMAKE_MODULE_PATH=$UE_CMAKE_DIR \
+		-DUE_IREE_BUILD_ROOT=$BUILD_DIR \
+		-DIREE_ENABLE_CPUINFO=OFF \
+		-DIREE_BUILD_TESTS=OFF \
+		-DIREE_BUILD_SAMPLES=OFF \
+		-DIREE_BUILD_BINDINGS_TFLITE=OFF \
+		-DIREE_BUILD_BINDINGS_TFLITE_JAVA=OFF \
+		-DIREE_BUILD_ALL_CHECK_TEST_MODULES=OFF \
+		-DIREE_HAL_DRIVER_DEFAULTS=OFF \
+		-DIREE_HAL_DRIVER_LOCAL_SYNC=ON \
+		-DIREE_HAL_DRIVER_LOCAL_TASK=ON \
+		-DIREE_TARGET_BACKEND_DEFAULTS=OFF \
+		-DIREE_TARGET_BACKEND_LLVM_CPU=ON \
+		-DIREE_ERROR_ON_MISSING_SUBMODULES=OFF \
+		-DIREE_BUILD_COMPILER=OFF \
+		-DIREE_ENABLE_THREADING=ON \
+		-DIREE_HAL_EXECUTABLE_LOADER_DEFAULTS=OFF \
+		-DIREE_HAL_EXECUTABLE_PLUGIN_DEFAULTS=OFF \
+		-DCMAKE_C_COMPILER=$C_COMPILER \
+		-DCMAKE_C_FLAGS=$COMPILER_FLAGS \
+		-DCMAKE_CXX_COMPILER=$CXX_COMPILER \
+		-DCMAKE_CXX_FLAGS=$COMPILER_FLAGS
+		
+	cmake --build $BUILD_DIR/iree-runtime-linux
+	
+	echo
+	echo Building Linux Runtime: Done
+	echo
+else
+	echo
+	echo Building Linux Runtime: Skipped
+	echo
+fi
diff --git a/Engine/Plugins/Experimental/NNERuntimeIREE/Source/ExternalBuild/Scripts/Linux/InstallLinux.sh b/Engine/Plugins/Experimental/NNERuntimeIREE/Source/ExternalBuild/Scripts/Linux/InstallLinux.sh
new file mode 100755
--- /dev/null
+++ b/Engine/Plugins/Experimental/NNERuntimeIREE/Source/ExternalBuild/Scripts/Linux/InstallLinux.sh
@@ -0,0 +1,171 @@
+#! /bin/bash
+
+#
+# This script copies all required files back into the plugin folder tree and adds them to a new Perforce CL.
+#
+# If an intermediate step fails, the script exits.
+#
+
+set -Eeuo pipefail
+
+if [ -z "$1" ]; then
+	echo "Usage: [Script] [Working dir]"
+	exit 1
+fi
+
+WORKING_DIR=$1
+
+# Install signal handler for ERR signal
+on_error() {
+	local exit_code=$?
+	local line=${BASH_LINENO[0]}
+
+	echo "Error: command '${BASH_COMMAND}' failed at line ${line} with exit code ${exit_code}" >&2
+	exit "$exit_code"
+}
+trap on_error ERR
+
+# Common rsync options
+# note: -L to copy "resolved" symlinks (e.g. ld.lld)
+RSYNC_OPTIONS=(-a -L)
+
+sync() { rsync "${RSYNC_OPTIONS[@]}" -- "$@"; }
+
+echo =========================================
+echo ========= Installing packages ===========
+echo =========================================
+
+die() { echo "Error: $*" >&2; exit 1; }
+
+packages=(rsync patchelf)
+sudo apt-get -y --ignore-missing install "${packages[@]}" >/dev/null 2>&1
+
+commands=(realpath awk rsync patchelf p4)
+for p in "${commands[@]}"; do
+	command -v "$p" >/dev/null 2>&1 || die "$p is required but not found."
+done
+
+echo Successfully checked for required commands.
+
+
+echo =========================================
+echo ============ Build config ===============
+echo =========================================
+
+BUILD_DIR="$WORKING_DIR/iree-org"
+BUILD_SCRIPT_DIR="$(dirname -- "$(realpath -- "$0")")/"
+UE_SOURCE_DIR="$BUILD_SCRIPT_DIR/../../Source"
+UE_PLUGIN_ROOT="$BUILD_SCRIPT_DIR/../../../.."
+
+echo Plugin root dir: "$UE_PLUGIN_ROOT"
+echo Build dir: "$BUILD_DIR"
+echo Source dir: "$UE_SOURCE_DIR"
+
+
+echo =========================================
+echo ======== Copying Compiler Files =========
+echo =========================================
+
+TARGET_IREE_FOLDER="$UE_PLUGIN_ROOT/Source/ThirdParty/IREE"
+BINARIES_FOLDER="$UE_PLUGIN_ROOT/Binaries/ThirdParty/IREE/Linux"
+mkdir -p "$BINARIES_FOLDER" 2>/dev/null
+
+sync "$BUILD_DIR/iree-compiler/llvm-project/bin/ld.lld" "$BINARIES_FOLDER/"
+sync "$BUILD_DIR/iree-compiler/tools/iree-compile" "$BINARIES_FOLDER/"
+sync "$BUILD_DIR/iree-compiler/lib/libIREECompiler.so" "$BINARIES_FOLDER/"
+
+# Add the ./ folder to the compilers rpath to find the shared lib in the same folder
+patchelf --set-rpath '$ORIGIN' "$BINARIES_FOLDER/iree-compile"
+
+echo
+echo Copying Compiler Files: Done
+echo
+
+
+echo =========================================
+echo ====== Copying NNEMlirTools Files =======
+echo =========================================
+
+MLIR_TOOLS_INCLUDE_DIR="$UE_PLUGIN_ROOT/Source/ThirdParty/NNEMlirTools/Internal"
+MLIR_TOOLS_BINARIES_DIR="$UE_PLUGIN_ROOT/Binaries/ThirdParty/NNEMlirTools/Linux"
+mkdir -p "$MLIR_TOOLS_INCLUDE_DIR" 2>/dev/null
+mkdir -p "$MLIR_TOOLS_BINARIES_DIR" 2>/dev/null
+
+sync "$BUILD_DIR/NNEMlirTools/libNNEMlirTools.so" "$MLIR_TOOLS_BINARIES_DIR/"
+sync "$UE_SOURCE_DIR/NNEMlirTools/NNERuntimeIREE/Include/NNEMlirTools.h" "$MLIR_TOOLS_INCLUDE_DIR/"
+sync "$UE_SOURCE_DIR/NNEMlirTools/NNERuntimeIREE/Include/NNEMlirTools_cxx_api.h" "$MLIR_TOOLS_INCLUDE_DIR/"
+
+echo
+echo Copying NNEMlirTools Files: Done
+echo
+
+
+echo =========================================
+echo =========== Copying Linux Files ===========
+echo =========================================
+
+LIBRARIES_FOLDER="$TARGET_IREE_FOLDER/Lib/Linux"
+mkdir -p "$LIBRARIES_FOLDER" 2>/dev/null
+
+sync "$BUILD_DIR/iree-runtime-linux/iree/build_tools/third_party/flatcc/libflatcc_parsing.a" "$LIBRARIES_FOLDER/"
+sync "$BUILD_DIR/iree-runtime-linux/libireert.a" "$LIBRARIES_FOLDER/"
+
+echo
+echo Copying Linux Files: Done
+echo
+
+
+echo =========================================
+echo ============== Save to CL ===============
+echo =========================================
+
+# Check for P4 CLI client
+if ! command -v p4 >/dev/null 2>&1; then
+	echo "Error: Perforce CLI (p4) not found in PATH." >&2
+	exit 1
+fi
+
+# Check connection to server
+if ! p4 -G info >/dev/null 2>&1; then
+	echo "Error: Unable to reach Perforce server." >&2
+	exit 1
+fi
+
+# Check logged in
+if ! p4 login -s >/dev/null 2>&1; then
+	echo "Starting interactive login..."
+	if ! p4 login; then
+		echo "Error: Login failed." >&2
+		exit 1
+	fi
+fi
+
+CL_DESC="IREE Linux install"
+
+CL="$(
+	p4 --field "Description=${CL_DESC}" --field "Files=" change -o \
+	| p4 change -i \
+	| awk '/^Change [0-9]+/ {print $2; exit}'
+)"
+
+if [[ -z "${CL}" ]]; then
+	echo "Failed to create change list (no CL number returned)." >&2
+	exit 1
+fi
+
+echo "Created changelist #${CL}"
+
+BINARIES_FOLDER="$(realpath -e -- "$BINARIES_FOLDER")"
+LIBRARIES_FOLDER="$(realpath -e -- "$LIBRARIES_FOLDER")"
+MLIR_TOOLS_INCLUDE_DIR="$(realpath -e -- "$MLIR_TOOLS_INCLUDE_DIR")"
+MLIR_TOOLS_BINARIES_DIR="$(realpath -e -- "$MLIR_TOOLS_BINARIES_DIR")"
+
+# Reconcile modified directories
+p4 reconcile -a -e -d -c "$CL" "$BINARIES_FOLDER/..."
+p4 reconcile -a -e -d -c "$CL" "$LIBRARIES_FOLDER/..."
+p4 reconcile -a -e -d -c "$CL" "$MLIR_TOOLS_INCLUDE_DIR/..."
+p4 reconcile -a -e -d -c "$CL" "$MLIR_TOOLS_BINARIES_DIR/..."
+
+echo
+echo Save to CL: Done
+echo
\ No newline at end of file
diff --git a/Engine/Plugins/Experimental/NNERuntimeIREE/Source/ExternalBuild/Scripts/Mac/BuildAndInstallMac.sh b/Engine/Plugins/Experimental/NNERuntimeIREE/Source/ExternalBuild/Scripts/Mac/BuildAndInstallMac.sh
new file mode 100755
--- /dev/null
+++ b/Engine/Plugins/Experimental/NNERuntimeIREE/Source/ExternalBuild/Scripts/Mac/BuildAndInstallMac.sh
@@ -0,0 +1,39 @@
+#! /bin/bash
+
+#
+# Script that calls build and install scripts together with the same arguments.
+#
+# If any intermediate step fails, the script (should) exit.
+#
+
+set -Eeuo pipefail
+
+trap 'echo; echo "Script failed. Exiting" >&2; exit 1' ERR
+
+SCRIPT_DIR="$(cd -- "$(dirname -- "${BASH_SOURCE[0]}")" &>/dev/null && pwd)"
+
+# For production be aware that the build path is part of the IREE error messages
+# So choose it with care
+DEFAULT_WORKING_DIR="$SCRIPT_DIR/../../Build"
+
+if [ -z "${1-}" ]; then
+	mkdir -p "$DEFAULT_WORKING_DIR"
+	WORKING_DIR="$DEFAULT_WORKING_DIR"
+else
+	WORKING_DIR="$1"
+	if [ ! -d "$WORKING_DIR" ]; then
+		echo "Error: '$WORKING_DIR' does not exist." >&2
+		exit 1
+	fi
+fi
+
+WORKING_DIR="$(cd "$WORKING_DIR" && pwd)"
+
+echo
+echo "Working dir: \"$WORKING_DIR\""
+echo
+
+"${SCRIPT_DIR}/BuildMac.sh" "$WORKING_DIR"
+"${SCRIPT_DIR}/InstallMac.sh" "$WORKING_DIR"
+
+exit 0
\ No newline at end of file
diff --git a/Engine/Plugins/Experimental/NNERuntimeIREE/Source/ExternalBuild/Scripts/Mac/BuildMac.sh b/Engine/Plugins/Experimental/NNERuntimeIREE/Source/ExternalBuild/Scripts/Mac/BuildMac.sh
new file mode 100755
--- /dev/null
+++ b/Engine/Plugins/Experimental/NNERuntimeIREE/Source/ExternalBuild/Scripts/Mac/BuildMac.sh
@@ -0,0 +1,218 @@
+#! /bin/bash
+
+#
+# This script builds all the required targets for the NNERuntimeIREE plugin.
+#
+# If an intermediate step fails, the script exits.
+#
+
+set -euo pipefail
+
+if [ -z "$1" ]; then
+	echo "Usage: [Script] [Working dir]"
+	exit 1
+fi
+
+WORKING_DIR=$1
+
+BUILD_DIR=$WORKING_DIR/iree-org
+BUILD_SCRIPT_DIR="$(dirname -- "$(realpath -- "$0")")/"
+UE_SOURCE_DIR=$BUILD_SCRIPT_DIR/../../Source
+UE_PATCHES_DIR=$BUILD_SCRIPT_DIR/../../Patch
+UE_CMAKE_DIR=$BUILD_SCRIPT_DIR/../../CMake
+UE_PLUGIN_ROOT=$BUILD_SCRIPT_DIR/../../../..
+
+echo Plugin root dir: "$UE_PLUGIN_ROOT"
+echo Build dir: "$BUILD_DIR"
+echo Source dir: "$UE_SOURCE_DIR"
+
+IREE_GIT_REPOSITORY=[REDACTED_PROFILE_URL]/iree.git
+IREE_GIT_COMMIT=v3.5.0
+IREE_THIRD_PARTY_LIBRARIES=("flatcc" "llvm-project" "stablehlo" "torch-mlir" "benchmark" "spirv_cross")
+IREE_COMPILER_VERSION_STRING="IREE-for-UE"
+
+BUILD_TYPE=Release
+
+# Unreal Shader compiler plugin
+IREE_CMAKE_PLUGIN_PATH=../../Iree/Compiler/Plugins/Target/UnrealShader
+
+SPIRV_CROSS_DIR=$BUILD_DIR/iree/third_party/spirv_cross
+
+OSX_TARGET=13.0
+ARCHITECTURES="arm64;x86_64"
+
+
+echo =========================================
+echo ============= Copy Source ===============
+echo =========================================
+
+mkdir -p "$WORKING_DIR/iree/compiler"
+cp -a "$UE_SOURCE_DIR/Iree/Compiler/." "$WORKING_DIR/iree/compiler/" > /dev/null
+
+echo
+echo Copied IREE compiler plugin.
+echo
+
+
+echo =========================================
+echo ============= Cloning IREE  =============
+echo =========================================
+
+if [ ! -d "$BUILD_DIR/iree" ]; then
+	mkdir $BUILD_DIR 2>/dev/null
+	cd $BUILD_DIR
+	git clone -n $IREE_GIT_REPOSITORY
+	cd iree
+	echo Using IREE git commit $IREE_GIT_COMMIT
+	git checkout $IREE_GIT_COMMIT
+	cd third_party
+	for D in */; do
+		if [ ! -f "$D/CMakeLists.txt" ]; then
+			echo -n "" > "$D/CMakeLists.txt"
+		fi
+	done
+	for L in ${IREE_THIRD_PARTY_LIBRARIES[@]}; do
+		if [ -d $L ]; then
+			rm "$L/CMakeLists.txt"
+			git submodule update --init -- $L
+			if [ -d "$L/third-party" ]; then
+				rm -R -f "$L/third-party"
+			fi
+		fi
+	done
+
+	cd ../../..
+
+	echo Check for $SPIRV_CROSS_DIR
+	if [ -d "$SPIRV_CROSS_DIR" ]; then
+		echo Apply git patch to spirv_cross
+		cd $SPIRV_CROSS_DIR
+		git apply $UE_PATCHES_DIR/spirv_cross.patch
+		cd ../../../..
+		echo Done.
+	fi
+
+	echo
+	echo Cloning IREE: Done
+	echo
+else
+	echo
+	echo Cloning IREE: Skipped
+	echo
+fi
+
+echo =========================================
+echo ======== Building IREE Compiler =========
+echo =========================================
+
+if [ ! -d "$BUILD_DIR/iree-compiler" ]; then
+	cmake -G Ninja -B $BUILD_DIR/iree-compiler $BUILD_DIR/iree \
+		-DCMAKE_BUILD_TYPE=$BUILD_TYPE \
+		-DIREE_CMAKE_PLUGIN_PATHS=$IREE_CMAKE_PLUGIN_PATH \
+		-DIREE_ENABLE_CPUINFO=OFF \
+		-DIREE_BUILD_TESTS=OFF \
+		-DIREE_BUILD_SAMPLES=OFF \
+		-DIREE_BUILD_BINDINGS_TFLITE=OFF \
+		-DIREE_BUILD_BINDINGS_TFLITE_JAVA=OFF \
+		-DIREE_BUILD_ALL_CHECK_TEST_MODULES=OFF \
+		-DIREE_HAL_DRIVER_DEFAULTS=OFF \
+		-DIREE_HAL_DRIVER_LOCAL_SYNC=ON \
+		-DIREE_HAL_DRIVER_LOCAL_TASK=ON \
+		-DIREE_TARGET_BACKEND_DEFAULTS=OFF \
+		-DIREE_TARGET_BACKEND_LLVM_CPU=ON \
+		-DIREE_ERROR_ON_MISSING_SUBMODULES=OFF \
+		-DCMAKE_OSX_ARCHITECTURES=$ARCHITECTURES \
+		-DCMAKE_OSX_DEPLOYMENT_TARGET=$OSX_TARGET \
+		-DIREE_EMBEDDED_RELEASE_INFO=ON \
+		-DIREE_RELEASE_VERSION=$IREE_COMPILER_VERSION_STRING \
+		-DIREE_RELEASE_REVISION=$IREE_GIT_COMMIT
+		
+	cmake --build $BUILD_DIR/iree-compiler
+	
+	echo
+	echo Building Compiler: Done
+	echo
+else
+	echo
+	echo Building Compiler: Incremental build... Please do not use for production.
+	echo
+
+	cmake --build $BUILD_DIR/iree-compiler
+fi
+
+
+echo =========================================
+echo ======== Building NNEMlirTools ==========
+echo =========================================
+
+if [ ! -d "$BUILD_DIR/NNEMlirTools" ]; then
+	cmake -G Ninja -B $BUILD_DIR/NNEMlirTools $UE_SOURCE_DIR/NNEMlirTools \
+		-DCMAKE_BUILD_TYPE=$BUILD_TYPE \
+		-DUE_IREE_BUILD_ROOT=$BUILD_DIR \
+		-DIREE_ENABLE_CPUINFO=OFF \
+		-DIREE_BUILD_TESTS=OFF \
+		-DIREE_BUILD_SAMPLES=OFF \
+		-DIREE_BUILD_BINDINGS_TFLITE=OFF \
+		-DIREE_BUILD_BINDINGS_TFLITE_JAVA=OFF \
+		-DIREE_BUILD_ALL_CHECK_TEST_MODULES=OFF \
+		-DIREE_HAL_DRIVER_DEFAULTS=OFF \
+		-DIREE_HAL_DRIVER_LOCAL_SYNC=OFF \
+		-DIREE_HAL_DRIVER_LOCAL_TASK=OFF \
+		-DIREE_TARGET_BACKEND_DEFAULTS=OFF \
+		-DIREE_TARGET_BACKEND_LLVM_CPU=OFF \
+		-DIREE_ERROR_ON_MISSING_SUBMODULES=OFF \
+		-DCMAKE_OSX_ARCHITECTURES=$ARCHITECTURES \
+		-DCMAKE_OSX_DEPLOYMENT_TARGET=$OSX_TARGET
+		
+	cmake --build $BUILD_DIR/NNEMlirTools
+	
+	echo
+	echo Building NNEMlirTools: Done
+	echo
+else
+	echo
+	echo Building NNEMlirTools: Incremental build... Please do not use for production.
+	echo
+
+	cmake --build $BUILD_DIR/NNEMlirTools
+fi
+
+
+echo =========================================
+echo ========= Building Mac Runtime ==========
+echo =========================================
+
+if [ ! -d "$BUILD_DIR/iree-runtime-mac" ]; then
+	cmake -G Ninja -B $BUILD_DIR/iree-runtime-mac $UE_SOURCE_DIR/Iree/Runtime \
+		-DCMAKE_BUILD_TYPE=$BUILD_TYPE \
+		-DCMAKE_MODULE_PATH=$UE_CMAKE_DIR \
+		-DUE_IREE_BUILD_ROOT=$BUILD_DIR \
+		-DIREE_ENABLE_CPUINFO=OFF \
+		-DIREE_BUILD_TESTS=OFF \
+		-DIREE_BUILD_SAMPLES=OFF \
+		-DIREE_BUILD_BINDINGS_TFLITE=OFF \
+		-DIREE_BUILD_BINDINGS_TFLITE_JAVA=OFF \
+		-DIREE_BUILD_ALL_CHECK_TEST_MODULES=OFF \
+		-DIREE_HAL_DRIVER_DEFAULTS=OFF \
+		-DIREE_HAL_DRIVER_LOCAL_SYNC=ON \
+		-DIREE_HAL_DRIVER_LOCAL_TASK=ON \
+		-DIREE_TARGET_BACKEND_DEFAULTS=OFF \
+		-DIREE_TARGET_BACKEND_LLVM_CPU=ON \
+		-DIREE_ERROR_ON_MISSING_SUBMODULES=OFF \
+		-DIREE_BUILD_COMPILER=OFF \
+		-DIREE_ENABLE_THREADING=ON \
+		-DIREE_HAL_EXECUTABLE_LOADER_DEFAULTS=OFF \
+		-DIREE_HAL_EXECUTABLE_PLUGIN_DEFAULTS=OFF \
+		-DCMAKE_OSX_ARCHITECTURES=$ARCHITECTURES \
+		-DCMAKE_OSX_DEPLOYMENT_TARGET=$OSX_TARGET
+		
+	cmake --build $BUILD_DIR/iree-runtime-mac
+	
+	echo
+	echo Building Mac Runtime: Done
+	echo
+else
+	echo
+	echo Building Mac Runtime: Skipped
+	echo
+fi
diff --git a/Engine/Plugins/Experimental/NNERuntimeIREE/Source/ExternalBuild/Scripts/Mac/InstallMac.sh b/Engine/Plugins/Experimental/NNERuntimeIREE/Source/ExternalBuild/Scripts/Mac/InstallMac.sh
new file mode 100755
--- /dev/null
+++ b/Engine/Plugins/Experimental/NNERuntimeIREE/Source/ExternalBuild/Scripts/Mac/InstallMac.sh
@@ -0,0 +1,166 @@
+#! /bin/bash
+
+#
+# This script copies all required files back into the plugin folder tree and adds them to a new Perforce CL.
+#
+# If an intermediate step fails, the script exits.
+#
+
+set -Eeuo pipefail
+
+if [ -z "$1" ]; then
+	echo "Usage: [Script] [Working dir]"
+	exit 1
+fi
+
+WORKING_DIR=$1
+
+BUILD_DIR="$WORKING_DIR/iree-org"
+BUILD_SCRIPT_DIR="$(dirname -- "$(realpath -- "$0")")/"
+UE_SOURCE_DIR="$BUILD_SCRIPT_DIR/../../Source"
+UE_PLUGIN_ROOT="$BUILD_SCRIPT_DIR/../../../.."
+
+echo Plugin root dir: "$UE_PLUGIN_ROOT"
+echo Build dir: "$BUILD_DIR"
+echo Source dir: "$UE_SOURCE_DIR"
+
+on_error() {
+	local exit_code=$?
+	local line=${BASH_LINENO[0]}
+
+	echo "Error: command '${BASH_COMMAND}' failed at line ${line} with exit code ${exit_code}" >&2
+	exit "$exit_code"
+}
+trap on_error ERR
+
+# Common rsync options
+# note: -L to copy "resolved" symlinks (e.g. ld.lld)
+RSYNC_OPTIONS=(-a -L)
+
+sync() { rsync "${RSYNC_OPTIONS[@]}" -- "$@"; }
+
+echo =========================================
+echo ========= INSTALLING PACKAGES ===========
+echo =========================================
+
+die() { echo "Error: $*" >&2; exit 1; }
+
+commands=(realpath awk rsync p4)
+for p in "${commands[@]}"; do
+	command -v "$p" >/dev/null 2>&1 || die "$p is required but not found."
+done
+
+echo Successfully checked for required commands.
+
+echo =========================================
+echo ======== Copying Compiler Files =========
+echo =========================================
+
+TARGET_IREE_FOLDER="$UE_PLUGIN_ROOT/Source/ThirdParty/IREE"
+BINARIES_FOLDER="$UE_PLUGIN_ROOT/Binaries/ThirdParty/IREE/Mac"
+mkdir -p "$BINARIES_FOLDER" 2>/dev/null
+
+sync "$BUILD_DIR/iree-compiler/llvm-project/bin/ld.lld" "$BINARIES_FOLDER/"
+sync "$BUILD_DIR/iree-compiler/llvm-project/bin/ld64.lld" "$BINARIES_FOLDER/"
+sync "$BUILD_DIR/iree-compiler/tools/iree-compile" "$BINARIES_FOLDER/"
+sync "$BUILD_DIR/iree-compiler/lib/libIREECompiler.dylib" "$BINARIES_FOLDER/"
+
+# Add the ./ folder to the compilers rpath to find the shared lib in the same folder
+install_name_tool -add_rpath "@executable_path/." "$BINARIES_FOLDER/iree-compile"
+install_name_tool -delete_rpath "$BUILD_DIR/iree-compiler/lib" "$BINARIES_FOLDER/iree-compile"
+
+echo
+echo Copying Compiler Files: Done
+echo
+
+
+echo =========================================
+echo ====== Copying NNEMlirTools Files =======
+echo =========================================
+
+MLIR_TOOLS_INCLUDE_DIR="$UE_PLUGIN_ROOT/Source/ThirdParty/NNEMlirTools/Internal"
+MLIR_TOOLS_BINARIES_DIR="$UE_PLUGIN_ROOT/Binaries/ThirdParty/NNEMlirTools/Mac"
+mkdir -p "$MLIR_TOOLS_INCLUDE_DIR" 2>/dev/null
+mkdir -p "$MLIR_TOOLS_BINARIES_DIR" 2>/dev/null
+
+sync "$BUILD_DIR/NNEMlirTools/libNNEMlirTools.dylib" "$MLIR_TOOLS_BINARIES_DIR/"
+sync "$UE_SOURCE_DIR/NNEMlirTools/NNERuntimeIREE/Include/NNEMlirTools.h" "$MLIR_TOOLS_INCLUDE_DIR/"
+sync "$UE_SOURCE_DIR/NNEMlirTools/NNERuntimeIREE/Include/NNEMlirTools_cxx_api.h" "$MLIR_TOOLS_INCLUDE_DIR/"
+
+echo
+echo Copying NNEMlirTools Files: Done
+echo
+
+
+echo =========================================
+echo =========== Copying Mac Files ===========
+echo =========================================
+
+LIBRARIES_FOLDER="$TARGET_IREE_FOLDER/Lib/Mac"
+mkdir -p "$LIBRARIES_FOLDER" 2>/dev/null
+
+sync "$BUILD_DIR/iree-runtime-mac/iree/build_tools/third_party/flatcc/libflatcc_parsing.a" "$LIBRARIES_FOLDER/"
+sync "$BUILD_DIR/iree-runtime-mac/libireert.a" "$LIBRARIES_FOLDER/"
+
+echo
+echo Copying Mac Files: Done
+echo
+
+
+echo =========================================
+echo ============== Save to CL ===============
+echo =========================================
+
+# Check for P4 CLI client
+if ! command -v p4 >/dev/null 2>&1; then
+	echo "Error: Perforce CLI (p4) not found in PATH." >&2
+	exit 1
+fi
+
+# Check connection to server
+if ! p4 -G info >/dev/null 2>&1; then
+	echo "Error: Unable to reach Perforce server." >&2
+	exit 1
+fi
+
+# Check logged in
+if ! p4 login -s >/dev/null 2>&1; then
+	echo "Starting interactive login..."
+	if ! p4 login; then
+		echo "Error: Login failed." >&2
+		exit 1
+	fi
+fi
+
+CL_DESC="IREE Mac install"
+
+CL=$(
+	{
+		printf 'Change: new\n'
+		printf 'Description:\n'
+		printf '%s' "$CL_DESC" | sed 's/^/\t/'
+	}   | p4 change -i \
+		| sed -n 's/^Change \([0-9][0-9]*\).*/\1/p' | head -n1
+)
+
+if [[ -z "${CL}" ]]; then
+	echo "Failed to create change list (no CL number returned)." >&2
+	exit 1
+fi
+
+echo "Created changelist #${CL}"
+
+BINARIES_FOLDER="$(cd -- "$BINARIES_FOLDER" && pwd -P)"
+LIBRARIES_FOLDER="$(cd -- "$LIBRARIES_FOLDER" && pwd -P)"
+MLIR_TOOLS_INCLUDE_DIR="$(cd -- "$MLIR_TOOLS_INCLUDE_DIR" && pwd -P)"
+MLIR_TOOLS_BINARIES_DIR="$(cd -- "$MLIR_TOOLS_BINARIES_DIR" && pwd -P)"
+
+# Reconcile modified directories
+p4 reconcile -a -e -d -c "$CL" "$BINARIES_FOLDER/..."
+p4 reconcile -a -e -d -c "$CL" "$LIBRARIES_FOLDER/..."
+p4 reconcile -a -e -d -c "$CL" "$MLIR_TOOLS_INCLUDE_DIR/..."
+p4 reconcile -a -e -d -c "$CL" "$MLIR_TOOLS_BINARIES_DIR/..."
+
+echo
+echo Save to CL: Done
+echo
\ No newline at end of file
diff --git a/Engine/Plugins/Experimental/NNERuntimeIREE/Source/ExternalBuild/Scripts/Windows/BuildAndInstallWindows.bat b/Engine/Plugins/Experimental/NNERuntimeIREE/Source/ExternalBuild/Scripts/Windows/BuildAndInstallWindows.bat
new file mode 100755
--- /dev/null
+++ b/Engine/Plugins/Experimental/NNERuntimeIREE/Source/ExternalBuild/Scripts/Windows/BuildAndInstallWindows.bat
@@ -0,0 +1,34 @@
+@ECHO off
+SETLOCAL
+
+::
+:: Script that calls build and install scripts together with the same arguments.
+::
+:: If any intermediate step fails, the script (should) exit. Without arguments it will build in
+:: NNERuntimeIREE/Source/ExternalBuild/Build and without additional platforms, e.g. Windows only.
+::
+
+:: Parse arguments (they will be available to build and install scripts)
+CALL "%~dp0_ParseArgs.bat" %* || EXIT /B %ERRORLEVEL%
+
+:: (Debug)print arguments
+ECHO:
+ECHO Working dir: "%WORKING_DIR%"
+for %%P in (%PLATFORM_LIST:,= %) do (
+	ECHO Building for %%P
+)
+ECHO:
+
+CALL "%~dp0BuildWindows.bat"
+IF ERRORLEVEL 1 GOTO :fail
+
+CALL "%~dp0InstallWindows.bat"
+IF ERRORLEVEL 1 GOTO :fail
+
+ENDLOCAL
+EXIT /B 0
+
+:fail
+ECHO.
+ECHO Build and Install script failed. Exiting 1>&2
+ENDLOCAL & EXIT /B 1
\ No newline at end of file
diff --git a/Engine/Plugins/Experimental/NNERuntimeIREE/Source/ExternalBuild/Scripts/Windows/BuildWindows.bat b/Engine/Plugins/Experimental/NNERuntimeIREE/Source/ExternalBuild/Scripts/Windows/BuildWindows.bat
new file mode 100755
--- /dev/null
+++ b/Engine/Plugins/Experimental/NNERuntimeIREE/Source/ExternalBuild/Scripts/Windows/BuildWindows.bat
@@ -0,0 +1,284 @@
+@ECHO off
+SETLOCAL
+
+::
+:: This script builds all the required targets for the NNERuntimeIREE plugin.
+::
+:: If an intermediate step fails, the script exits.
+::
+
+:: Get arguments
+IF NOT DEFINED ARGS_PARSED (
+	CALL "%~dp0_ParseArgs.bat" %* || EXIT /B %ERRORLEVEL%
+)
+
+IF NOT EXIST "%WORKING_DIR%" (
+	ECHO:
+	ECHO Warning: Working directory "%WORKING_DIR%" does not exist, creating...
+	EHCO:
+
+	MKDIR "%WORKING_DIR%" 2>NUL
+)
+
+SET "BUILD_DIR=%WORKING_DIR%\iree-org"
+SET "BUILD_SCRIPT_DIR=%~dp0"
+SET "UE_SOURCE_DIR=%BUILD_SCRIPT_DIR%\..\..\Source"
+SET "UE_PATCHES_DIR=%BUILD_SCRIPT_DIR%\..\..\Patch"
+SET "UE_CMAKE_DIR=%BUILD_SCRIPT_DIR%\..\..\CMake"
+SET "UE_PLUGIN_ROOT=%BUILD_SCRIPT_DIR%\..\..\..\.."
+
+ECHO Plugin root dir: "%UE_PLUGIN_ROOT%"
+ECHO Build dir: "%BUILD_DIR%"
+ECHO Source dir: "%UE_SOURCE_DIR%"
+ECHO Additional runtime platform argument list: %PLATFORM_LIST%
+
+SET IREE_GIT_REPOSITORY=[REDACTED_PROFILE_URL]/iree.git
+SET IREE_GIT_COMMIT=v3.5.0
+SET IREE_THIRD_PARTY_LIBRARIES=flatcc,llvm-project,stablehlo,torch-mlir,benchmark,spirv_cross
+SET IREE_COMPILER_VERSION_STRING="IREE-for-UE"
+
+SET BUILD_TYPE=Release
+
+:: Unreal Shader compiler plugin
+SET "IREE_CMAKE_PLUGIN_PATH=..\..\Iree\Compiler\Plugins\Target\UnrealShader"
+
+SET "SPIRV_CROSS_DIR=%BUILD_DIR%\iree\third_party\spirv_cross"
+
+SET VC_TOOLSET=14.38
+SET WIN_SDK_FOR_IREE=10.0.22621.0
+SET WIN_SDK_FOR_RUNTIME=10.0.22621.0
+
+ECHO =========================================
+ECHO ============= Copy Source ===============
+ECHO =========================================
+
+:: Need to copy at least IREE compiler plugin to avoid path lengths exploding
+robocopy "%UE_SOURCE_DIR%\iree\compiler" "%WORKING_DIR%\iree\compiler" /E >NUL
+IF ERRORLEVEL 8 GOTO :fail
+
+ECHO:
+ECHO Copied IREE compiler plugin.
+ECHO:
+
+ECHO =========================================
+ECHO ============= Cloning IREE ==============
+ECHO =========================================
+
+IF NOT EXIST "%BUILD_DIR%\iree" (
+	MKDIR "%BUILD_DIR%" 2>NUL
+	CD "%BUILD_DIR%"
+
+	git clone -n %IREE_GIT_REPOSITORY%
+	IF ERRORLEVEL 1 GOTO :fail
+
+	CD iree
+	
+	ECHO Using IREE git commit '%IREE_GIT_COMMIT%'
+	git checkout %IREE_GIT_COMMIT%
+	IF ERRORLEVEL 1 GOTO :fail
+	
+	FOR /F "tokens=*" %%F IN ('DIR "%BUILD_DIR%\iree\third_party" /AD /B') DO (
+		IF NOT EXIST "%BUILD_DIR%\iree\third_party\%%F\CMakeLists.txt" COPY NUL "%BUILD_DIR%\iree\third_party\%%F\CMakeLists.txt" >NUL
+	)
+	
+	FOR %%L IN (%IREE_THIRD_PARTY_LIBRARIES%) DO (
+		IF EXIST "third_party\%%L" (
+			DEL "".\third_party\%%L\CMakeLists.txt"
+			git submodule update --init -- ""./third_party/%%L"
+			IF ERRORLEVEL 1 GOTO :fail
+			IF EXIST "third_party\%%L\third-party" RMDIR /S /Q "third_party\%%L\third-party"
+		)
+	)
+
+	CD ../..
+
+	:: Patch SPIRV-Cross
+	ECHO Check for "%SPIRV_CROSS_DIR%"
+	IF EXIST "%SPIRV_CROSS_DIR%" (
+		ECHO Apply git patch to spirv_cross
+		CD "%SPIRV_CROSS_DIR%"
+
+		git apply "%UE_PATCHES_DIR%\spirv_cross.patch"
+		IF ERRORLEVEL 1 GOTO :fail
+
+		CD ..\..\..\..
+		ECHO Done.
+	)
+
+	ECHO:
+	ECHO Cloning IREE: Done
+	ECHO:
+) ELSE (
+	ECHO:
+	ECHO Cloning IREE: Skipped
+	ECHO:
+)
+
+ECHO =========================================
+ECHO ======== Building IREE Compiler =========
+ECHO =========================================
+
+IF NOT EXIST "%BUILD_DIR%\iree-compiler" (
+	SETLOCAL
+	CALL "%VCVARSALL_PATH%" x64 %WIN_SDK_FOR_IREE% -vcvars_ver=%VC_TOOLSET%
+	cmake -G Ninja -B "%BUILD_DIR%\iree-compiler" "%BUILD_DIR%\iree" ^
+		-DCMAKE_BUILD_TYPE=%BUILD_TYPE% ^
+		-DIREE_CMAKE_PLUGIN_PATHS="%IREE_CMAKE_PLUGIN_PATH%" ^
+		-DIREE_ENABLE_CPUINFO=OFF ^
+		-DIREE_BUILD_TESTS=OFF ^
+		-DIREE_BUILD_SAMPLES=OFF ^
+		-DIREE_BUILD_BINDINGS_TFLITE=OFF ^
+		-DIREE_BUILD_BINDINGS_TFLITE_JAVA=OFF ^
+		-DIREE_BUILD_ALL_CHECK_TEST_MODULES=OFF ^
+		-DIREE_HAL_DRIVER_DEFAULTS=OFF ^
+		-DIREE_HAL_DRIVER_LOCAL_SYNC=ON ^
+		-DIREE_HAL_DRIVER_LOCAL_TASK=ON ^
+		-DIREE_TARGET_BACKEND_DEFAULTS=OFF ^
+		-DIREE_TARGET_BACKEND_LLVM_CPU=ON ^
+		-DIREE_ERROR_ON_MISSING_SUBMODULES=OFF ^
+		-DIREE_EMBEDDED_RELEASE_INFO=ON ^
+		-DIREE_RELEASE_VERSION=%IREE_COMPILER_VERSION_STRING% ^
+		-DIREE_RELEASE_REVISION=%IREE_GIT_COMMIT%
+	IF ERRORLEVEL 1 (ENDLOCAL & GOTO :fail )
+
+	cmake --build "%BUILD_DIR%\iree-compiler"
+	IF ERRORLEVEL 1 (ENDLOCAL & GOTO :fail )
+
+	ENDLOCAL
+
+	ECHO:
+	ECHO Building IREE Compiler: Done
+	ECHO:
+) ELSE (
+	ECHO:
+	ECHO Building IREE Compiler: Incremental build... Please do not use for Release build.
+	ECHO:
+
+	SETLOCAL
+	CALL "%VCVARSALL_PATH%" x64 %WIN_SDK_FOR_IREE% -vcvars_ver=%VC_TOOLSET%
+
+	cmake --build "%BUILD_DIR%\iree-compiler"
+	IF ERRORLEVEL 1 (ENDLOCAL & GOTO :fail )
+	
+	ENDLOCAL
+	ECHO:
+)
+
+ECHO =========================================
+ECHO ======== Building NNEMlirTools ==========
+ECHO =========================================
+
+IF NOT EXIST %BUILD_DIR%\NNEMlirTools (
+	SETLOCAL
+	CALL "%VCVARSALL_PATH%" x64 %WIN_SDK_FOR_IREE% -vcvars_ver=%VC_TOOLSET%
+	cmake -G Ninja -B "%BUILD_DIR%/NNEMlirTools" "%UE_SOURCE_DIR%\NNEMlirTools" ^
+		-DCMAKE_BUILD_TYPE=%BUILD_TYPE% ^
+		-DUE_IREE_BUILD_ROOT="%BUILD_DIR%" ^
+		-DIREE_ENABLE_CPUINFO=OFF ^
+		-DIREE_BUILD_TESTS=OFF ^
+		-DIREE_BUILD_SAMPLES=OFF ^
+		-DIREE_BUILD_BINDINGS_TFLITE=OFF ^
+		-DIREE_BUILD_BINDINGS_TFLITE_JAVA=OFF ^
+		-DIREE_BUILD_ALL_CHECK_TEST_MODULES=OFF ^
+		-DIREE_HAL_DRIVER_DEFAULTS=OFF ^
+		-DIREE_HAL_DRIVER_LOCAL_SYNC=OFF ^
+		-DIREE_HAL_DRIVER_LOCAL_TASK=OFF ^
+		-DIREE_TARGET_BACKEND_DEFAULTS=OFF ^
+		-DIREE_TARGET_BACKEND_LLVM_CPU=OFF ^
+		-DIREE_ERROR_ON_MISSING_SUBMODULES=OFF
+	IF ERRORLEVEL 1 (ENDLOCAL & GOTO :fail )
+
+	cmake --build "%BUILD_DIR%/NNEMlirTools"
+	IF ERRORLEVEL 1 (ENDLOCAL & GOTO :fail )
+
+	ENDLOCAL
+
+	ECHO:
+	ECHO Building NNEMlirTools: Done
+	ECHO:
+) ELSE (
+	ECHO:
+	ECHO Building NNEMlirTools: Incremental build... Please do not use for Release build.
+	ECHO:
+
+	SETLOCAL
+	CALL "%VCVARSALL_PATH%" x64 %WIN_SDK_FOR_IREE% -vcvars_ver=%VC_TOOLSET%
+
+	cmake --build "%BUILD_DIR%/NNEMlirTools"
+	IF ERRORLEVEL 1 (ENDLOCAL & GOTO :fail )
+
+	ENDLOCAL
+
+	ECHO:
+)
+
+ECHO =========================================
+ECHO ======= Building Windows Runtime ========
+ECHO =========================================
+
+:: Need to sanitize for cmake
+SET "UE_CMAKE_DIR_SANITIZED=%UE_CMAKE_DIR:\=/%"
+ECHO Sanitized path: "%UE_CMAKE_DIR_SANITIZED%"
+
+IF NOT EXIST "%BUILD_DIR%\iree-runtime-windows" (
+	SETLOCAL
+	CALL "%VCVARSALL_PATH%" x64 %WIN_SDK_FOR_RUNTIME% -vcvars_ver=%VC_TOOLSET%
+	cmake -G Ninja -B "%BUILD_DIR%\iree-runtime-windows" "%UE_SOURCE_DIR%\iree\runtime" ^
+		-DCMAKE_BUILD_TYPE=%BUILD_TYPE% ^
+		-DCMAKE_MODULE_PATH="%UE_CMAKE_DIR_SANITIZED%" ^
+		-DUE_IREE_BUILD_ROOT="%BUILD_DIR%" ^
+		-DIREE_ENABLE_CPUINFO=OFF ^
+		-DIREE_BUILD_TESTS=OFF ^
+		-DIREE_BUILD_SAMPLES=OFF ^
+		-DIREE_BUILD_BINDINGS_TFLITE=OFF ^
+		-DIREE_BUILD_BINDINGS_TFLITE_JAVA=OFF ^
+		-DIREE_BUILD_ALL_CHECK_TEST_MODULES=OFF ^
+		-DIREE_HAL_DRIVER_DEFAULTS=OFF ^
+		-DIREE_HAL_DRIVER_LOCAL_SYNC=ON ^
+		-DIREE_HAL_DRIVER_LOCAL_TASK=ON ^
+		-DIREE_TARGET_BACKEND_DEFAULTS=OFF ^
+		-DIREE_TARGET_BACKEND_LLVM_CPU=ON ^
+		-DIREE_ERROR_ON_MISSING_SUBMODULES=OFF ^
+		-DIREE_BUILD_COMPILER=OFF ^
+		-DIREE_ENABLE_THREADING=ON ^
+		-DIREE_HAL_EXECUTABLE_LOADER_DEFAULTS=OFF ^
+		-DIREE_HAL_EXECUTABLE_PLUGIN_DEFAULTS=OFF
+	IF ERRORLEVEL 1 (ENDLOCAL & GOTO :fail )
+
+	cmake --build "%BUILD_DIR%\iree-runtime-windows"
+	IF ERRORLEVEL 1 (ENDLOCAL & GOTO :fail )
+
+	ENDLOCAL
+
+	ECHO:
+	ECHO Building Windows Runtime: Done
+	ECHO:
+) ELSE (
+	ECHO:
+	ECHO Building Windows Runtime: Incremental build... Please do not use for Release build.
+	ECHO:
+
+	SETLOCAL
+	CALL "%VCVARSALL_PATH%" x64 %WIN_SDK_FOR_RUNTIME% -vcvars_ver=%VC_TOOLSET%
+
+	cmake --build "%BUILD_DIR%\iree-runtime-windows"
+	IF ERRORLEVEL 1 (ENDLOCAL & GOTO :fail )
+
+	ENDLOCAL
+
+	ECHO:
+)
+
+:: Run additional platform build scripts
+for %%i in (%PLATFORM_LIST%) do (
+	CALL "%UE_PLUGIN_ROOT%\..\..\..\Platforms\%%i\Plugins\Experimental\NNERuntimeIREE\Source\ExternalBuild\Scripts\Windows\BuildWindows.bat" "%BUILD_DIR%" "%BUILD_TYPE%" "%UE_CMAKE_DIR_SANITIZED%"
+	IF ERRORLEVEL 1 GOTO :fail
+)
+
+ENDLOCAL
+EXIT /B 0
+
+:fail
+ECHO.
+ECHO Build script failed. Exiting 1>&2
+ENDLOCAL & EXIT /B 1
\ No newline at end of file
diff --git a/Engine/Plugins/Experimental/NNERuntimeIREE/Source/ExternalBuild/Scripts/Windows/Clean.bat b/Engine/Plugins/Experimental/NNERuntimeIREE/Source/ExternalBuild/Scripts/Windows/Clean.bat
new file mode 100755
--- /dev/null
+++ b/Engine/Plugins/Experimental/NNERuntimeIREE/Source/ExternalBuild/Scripts/Windows/Clean.bat
@@ -0,0 +1,35 @@
+@ECHO off
+SETLOCAL
+
+::
+:: This script removes all the content of the folder WORKING_DIR.
+::
+:: PLEASE BE CAREFUL WITH THIS!
+::
+
+CALL "%~dp0_ParseArgs.bat" %* || EXIT /B %ERRORLEVEL%
+
+ECHO:
+ECHO Working dir: "%WORKING_DIR%"
+ECHO:
+
+SET "EMPTY_DIR=E:\empty"
+
+MKDIR "%EMPTY_DIR%"
+robocopy "%EMPTY_DIR%" "%WORKING_DIR%" /MIR /R:0 /W:0 /MT:64 /XJ /NJH /NJS /NP /NDL /NFL /LOG:NUL
+IF ERRORLEVEL 8 GOTO :fail
+
+RMDIR "%EMPTY_DIR%"
+
+ECHO Successfully removed build dir, bye!
+EXIT /B 0
+
+:fail
+ECHO.
+
+if EXIST "%EMPTY_DIR%" (
+	RMDIR "%EMPTY_DIR%"
+)
+
+ECHO Clean script failed. Exiting 1>&2
+ENDLOCAL & EXIT /B 1
\ No newline at end of file
diff --git a/Engine/Plugins/Experimental/NNERuntimeIREE/Source/ExternalBuild/Scripts/Windows/InstallWindows.bat b/Engine/Plugins/Experimental/NNERuntimeIREE/Source/ExternalBuild/Scripts/Windows/InstallWindows.bat
new file mode 100755
--- /dev/null
+++ b/Engine/Plugins/Experimental/NNERuntimeIREE/Source/ExternalBuild/Scripts/Windows/InstallWindows.bat
@@ -0,0 +1,266 @@
+@ECHO off
+SETLOCAL
+
+::
+:: This script copies all required files back into the plugin folder tree and adds them to a new Perforce CL.
+::
+:: If an intermediate step fails, the script exits.
+::
+
+:: Get arguments
+IF NOT DEFINED ARGS_PARSED (
+	CALL "%~dp0_ParseArgs.bat" %* || EXIT /B %ERRORLEVEL%
+)
+
+SET "BUILD_DIR=%WORKING_DIR%\iree-org"
+SET "BUILD_SCRIPT_DIR=%~dp0"
+
+SET "ROBOCOPY_ARGS=/A-:R /NFL /NDL /NJH /NJS /NP"
+
+:: Skip helper function(s)
+goto :main
+
+:: Copy files with robocopy and catch errors
+:CP
+	SETLOCAL
+	robocopy %* %ROBOCOPY_ARGS%
+	SET "RC=%ERRORLEVEL%"
+
+	IF %RC% LEQ 3 (
+		IF %RC% EQU 0 ( ECHO Up-to-date )
+		IF %RC% EQU 1 ( ECHO Copied successfully )
+		IF %RC% EQU 2 ( ECHO Warning: extras in destination )
+		IF %RC% EQU 3 ( ECHO Copied successfully; extras exist )
+
+		ECHO   "%~1" 1>&2
+		ECHO   "%~2" 1>&2
+		ENDLOCAL & EXIT /B 0
+	)
+
+	ECHO Robocopy failed ^(code %RC%^)
+	ECHO   "%~1" 1>&2
+	ECHO   "%~2" 1>&2
+	ENDLOCAL & EXIT /B 1
+
+:: P4 wrapper
+:P4
+	p4 %*
+	IF ERRORLEVEL 1 GOTO :fail
+	GOTO :eof
+
+:: Canonicalize paths for Perforce (no '.' or '..' etc.)
+:Canonicalize
+	FOR %%I IN ("%~1") DO SET "%~2=%%~fI"
+	GOTO :eof
+
+:: Main body
+:main
+
+CALL :Canonicalize "%BUILD_SCRIPT_DIR%\..\..\Include" UE_INCLUDE_DIR
+CALL :Canonicalize "%BUILD_SCRIPT_DIR%\..\..\Source" UE_SOURCE_DIR
+CALL :Canonicalize "%BUILD_SCRIPT_DIR%\..\..\..\.." UE_PLUGIN_ROOT
+
+ECHO Plugin root dir: "%UE_PLUGIN_ROOT%"
+ECHO Build dir: "%BUILD_DIR%"
+ECHO Source dir: "%UE_SOURCE_DIR%"
+ECHO Additional runtime platform argument list: "%PLATFORM_LIST%"
+
+ECHO:
+
+ECHO =========================================
+ECHO ======== Copying Compiler Files =========
+ECHO =========================================
+
+SET "TARGET_IREE_FOLDER=%UE_PLUGIN_ROOT%\Source\ThirdParty\IREE"
+SET "BINARIES_FOLDER=%UE_PLUGIN_ROOT%\Binaries\ThirdParty\IREE\Windows"
+
+CALL :CP "%BUILD_DIR%\iree-compiler\llvm-project\bin" "%BINARIES_FOLDER%" clang++.exe ld.lld.exe
+IF ERRORLEVEL 1 GOTO :fail
+
+CALL :CP "%BUILD_DIR%\iree-compiler\tools" "%BINARIES_FOLDER%" iree-compile.exe IREECompiler.dll
+IF ERRORLEVEL 1 GOTO :fail
+
+ECHO:
+ECHO Copying Compiler Files: Done
+ECHO:
+
+ECHO =========================================
+ECHO ========= Copying Header Files ==========
+ECHO =========================================
+
+SET "HEADERS_FOLDER=%TARGET_IREE_FOLDER%\Include"
+
+CALL :CP "%BUILD_DIR%\iree\third_party\flatcc\include\flatcc" "%HEADERS_FOLDER%\flatcc" *.h /MIR
+IF ERRORLEVEL 1 GOTO :fail
+
+CALL :CP "%BUILD_DIR%\iree\runtime\src\iree" "%HEADERS_FOLDER%\iree" *.h /MIR
+IF ERRORLEVEL 1 GOTO :fail
+
+CALL :CP "%UE_INCLUDE_DIR%\Clang" "%HEADERS_FOLDER%\Clang" *.h /MIR
+IF ERRORLEVEL 1 GOTO :fail
+
+CALL :CP "%BUILD_DIR%\iree-compiler\runtime\src\iree\base\internal\flatcc" "%HEADERS_FOLDER%\iree\base\internal\flatcc" dummy_reader.h dummy_verifier.h
+IF ERRORLEVEL 1 GOTO :fail
+
+CALL :CP "%BUILD_DIR%\iree-compiler\runtime\src\iree\schemas" "%HEADERS_FOLDER%\iree\schemas" executable_debug_info_reader.h executable_debug_info_verifier.h
+IF ERRORLEVEL 1 GOTO :fail
+
+CALL :CP "%BUILD_DIR%\iree-compiler\compiler\plugins\target\UnrealShader" "%HEADERS_FOLDER%\iree\schemas" unreal_executable_def_reader.h unreal_executable_def_verifier.h
+IF ERRORLEVEL 1 GOTO :fail
+
+ECHO:
+ECHO Copying Header Files: Done
+ECHO:
+
+ECHO =========================================
+ECHO ====== Copying NNEMlirTools Files =======
+ECHO =========================================
+
+SET "MLIR_TOOLS_INCLUDE_DIR=%UE_PLUGIN_ROOT%\Source\ThirdParty\NNEMlirTools\Internal"
+SET "MLIR_TOOLS_BINARIES_DIR=%UE_PLUGIN_ROOT%\Binaries\ThirdParty\NNEMlirTools\Win64"
+
+CALL :CP "%BUILD_DIR%\NNEMlirTools" "%MLIR_TOOLS_BINARIES_DIR%" NNEMlirTools.dll
+IF ERRORLEVEL 1 GOTO :fail
+
+CALL :CP "%UE_SOURCE_DIR%\NNEMlirTools\NNERuntimeIREE\Include" "%MLIR_TOOLS_INCLUDE_DIR%" *.h /MIR
+IF ERRORLEVEL 1 GOTO :fail
+
+ECHO:
+ECHO Copying NNEMlirTools Files: Done
+ECHO:
+
+ECHO =========================================
+ECHO ========= Copying Windows Files =========
+ECHO =========================================
+
+SET "LIBRARIES_FOLDER=%TARGET_IREE_FOLDER%\Lib\Windows"
+
+CALL :CP "%BUILD_DIR%\iree-runtime-windows\iree\build_tools\third_party\flatcc" "%LIBRARIES_FOLDER%" flatcc_parsing.lib
+IF ERRORLEVEL 1 GOTO :fail
+
+CALL :CP "%BUILD_DIR%\iree-runtime-windows" "%LIBRARIES_FOLDER%" ireert.lib
+IF ERRORLEVEL 1 GOTO :fail
+
+ECHO:
+ECHO Copying Windows Files: Done
+ECHO:
+
+ECHO =========================================
+ECHO == Copying Additional Platforms Files ===
+ECHO =========================================
+
+SETLOCAL ENABLEDELAYEDEXPANSION
+SET "ADDITIONAL_FOLDERS="
+
+for %%i in (%PLATFORM_LIST%) do (
+	CALL :Canonicalize "%UE_PLUGIN_ROOT%\..\..\..\Platforms\%%i\Plugins\Experimental\NNERuntimeIREE" UE_PLATFORM_PLUGIN_ROOT
+	SET "PLATFORM_LIBRARIES_FOLDER=!UE_PLATFORM_PLUGIN_ROOT!\Source\ThirdParty\IREE\Lib\%%i"
+	ECHO Platform Plugin root: "!UE_PLATFORM_PLUGIN_ROOT!"
+	ECHO Library dir: "!PLATFORM_LIBRARIES_FOLDER!"
+
+	SET "ADDITIONAL_FOLDERS=!ADDITIONAL_FOLDERS! !PLATFORM_LIBRARIES_FOLDER!"
+
+	CALL :CP "%BUILD_DIR%\iree-runtime-%%i\iree\build_tools\third_party\flatcc" "!PLATFORM_LIBRARIES_FOLDER!" *.lib *.a
+	IF ERRORLEVEL 1 GOTO :fail
+
+	CALL :CP "%BUILD_DIR%\iree-runtime-%%i" "!PLATFORM_LIBRARIES_FOLDER!" *.lib *.a
+	IF ERRORLEVEL 1 GOTO :fail
+
+	SET "PLATFORM_HEADERS_FOLDER_SRC=!UE_PLATFORM_PLUGIN_ROOT!\Source\ExternalBuild\Include"
+	IF EXIST !PLATFORM_HEADERS_FOLDER_SRC! (
+		SET "PLATFORM_HEADERS_FOLDER_DST=!UE_PLATFORM_PLUGIN_ROOT!\Source\ThirdParty\IREE\Include"
+
+		ECHO:
+		ECHO Copy headers:
+		ECHO  from "!PLATFORM_HEADERS_FOLDER_SRC!"
+		ECHO  to   "!PLATFORM_HEADERS_FOLDER_DST!"
+
+		CALL :CP "!PLATFORM_HEADERS_FOLDER_SRC!" "!PLATFORM_HEADERS_FOLDER_DST!" *.h /MIR
+		IF ERRORLEVEL 1 GOTO :fail
+	)
+)
+
+ECHO All dirs to reconcile: "%ADDITIONAL_FOLDERS%"
+
+ECHO:
+ECHO Copying Additional Platforms Files: Done
+ECHO:
+
+ECHO =========================================
+ECHO ============== Save to CL ===============
+ECHO =========================================
+
+:: Check for P4 CLI client
+WHERE p4 >NUL 2>&1
+IF errorlevel 1 (
+	ECHO Error: Perforce CLI ^(p4.exe^) not found in PATH. 1>&2
+	GOTO :fail
+)
+
+:: Check connection to server
+p4 -G info >NUL 2>&1
+IF errorlevel 1 (
+	ECHO Error: Unable to reach Perforce server. 1>&2
+	GOTO :fail
+)
+
+:: Check logged in
+p4 login -s >NUL 2>&1
+IF errorlevel 1 (
+	ECHO Starting interactive login...
+	p4 login
+	if errorlevel 1 (
+		ECHO Error: Login failed. 1>&2
+		GOTO :fail
+	)
+)
+
+SET "CL_DESC=IREE install"
+
+:: Create a new CL with the description and NO files in it
+for /f "tokens=2" %%C in ('
+	p4 --field "Description=%CL_DESC%" --field "Files=" change -o ^
+	^| p4 change -i ^
+	^| findstr /r "^Change [0-9][0-9]*"
+') do (
+	SET "CL=%%C"
+)
+
+ECHO Created changelist #%CL%
+
+:: Reconcile modified directories
+p4 reconcile -a -e -d -c %CL% "%BINARIES_FOLDER%\..."
+IF ERRORLEVEL 1 GOTO :fail
+
+p4 reconcile -a -e -d -c %CL% "%HEADERS_FOLDER%\..."
+IF ERRORLEVEL 1 GOTO :fail
+
+p4 reconcile -a -e -d -c %CL% "%LIBRARIES_FOLDER%\..."
+IF ERRORLEVEL 1 GOTO :fail
+
+p4 reconcile -a -e -d -c %CL% "%MLIR_TOOLS_INCLUDE_DIR%\..."
+IF ERRORLEVEL 1 GOTO :fail
+
+p4 reconcile -a -e -d -c %CL% "%MLIR_TOOLS_BINARIES_DIR%\..."
+IF ERRORLEVEL 1 GOTO :fail
+
+FOR %%F IN (%ADDITIONAL_FOLDERS%) DO (
+	p4 reconcile -a -e -d -c %CL% "%%F%\..."
+	IF ERRORLEVEL 1 GOTO :fail
+)
+
+:: Note: this is required for the ADDITIONAL_FOLDERS list to be up-to-date
+ENDLOCAL
+
+ECHO:
+ECHO Save to CL: Done
+ECHO:
+
+ENDLOCAL
+ECHO All work done successfully, bye!
+EXIT /B 0
+
+:fail
+ECHO.
+ECHO Install script failed. Exiting 1>&2
+ENDLOCAL & EXIT /B 1
\ No newline at end of file
diff --git a/Engine/Plugins/Experimental/NNERuntimeIREE/Source/ExternalBuild/Scripts/Windows/_ParseArgs.bat b/Engine/Plugins/Experimental/NNERuntimeIREE/Source/ExternalBuild/Scripts/Windows/_ParseArgs.bat
new file mode 100755
--- /dev/null
+++ b/Engine/Plugins/Experimental/NNERuntimeIREE/Source/ExternalBuild/Scripts/Windows/_ParseArgs.bat
@@ -0,0 +1,96 @@
+@ECHO OFF
+SETLOCAL EnableExtensions EnableDelayedExpansion
+
+:: Usage:
+::   _ParseArg.bat [-dir WorkingDir] [-platform "Platform1 Platform2 ..."]
+
+:: Defaults
+FOR %%I IN ("%~dp0..\..\Build") DO SET "WORKING_DIR=%%~fI"
+SET "PLATFORM_LIST="
+
+:: UE root (sanitized)
+FOR %%I IN ("%~dp0..\..\..\..\..\..\..\..") DO SET "UE5_ROOT=%%~fI"
+
+:: Parse arguments: -dir and -platform
+:parse
+IF "%~1"=="" GOTO :after_parse
+
+:: -dir PATH
+IF /I "%~1"=="-dir" (
+	IF "%~2"=="" ( ECHO Error: -dir requires a path.>&2 & EXIT /B 2 )
+	SET "WORKING_DIR=%~2"
+	SHIFT & SHIFT & GOTO :parse
+)
+
+:: -platform LIST (LIST can be "All" or "[Platform1] [Platform2] [Platform3]")
+IF /I "%~1"=="-platform" (
+	SET "ENGINE_PLATFORMS_DIR=%UE5_ROOT%\Engine\Platforms"
+	IF NOT EXIST "%ENGINE_PLATFORMS_DIR%\" EXIT /B 0
+
+	IF "%~2"=="" ( ECHO Error: -platform requires a value.>&2 & EXIT /B 2 )
+	IF /I "%~2" == "All" (
+		CALL :discover_all
+	) else (
+		CALL :parse_platforms "%~2" || EXIT /B 1
+	)
+	SHIFT & SHIFT & GOTO :parse
+)
+
+ECHO Error: Unknown option: %~1>&2
+EXIT /B 2
+
+:after_parse
+:: Canonicalize WORKING_DIR
+FOR %%I IN ("%WORKING_DIR%") DO SET "WORKING_DIR=%%~fI"
+
+:: Summary
+@REM ECHO WORKING_DIR   = "%WORKING_DIR%"
+@REM ECHO UE5_ROOT      = "%UE5_ROOT%"
+@REM ECHO PLATFORM_LIST = "%PLATFORM_LIST%"
+
+ENDLOCAL & (
+	SET "WORKING_DIR=%WORKING_DIR%"
+	SET "UE5_ROOT=%UE5_ROOT%"
+	SET "PLATFORM_LIST=%PLATFORM_LIST%"
+	SET ARGS_PARSED=1
+)
+EXIT /B 0
+
+
+:parse_platforms
+:: Expand commas to spaces and add tokens, and handle All
+SET "arg=%~1"
+SET "arg=%arg:"=%"
+SET "arg=%arg:,= %"
+FOR %%P IN (%arg%) DO (
+	CALL :add_platform "%%~P" || (
+		ECHO Error: Platform "%%~P" has no build support or does not exist
+		EXIT /B 1
+	)
+)
+EXIT /B 0
+
+
+:add_platform
+:: Add %~1 to PLATFORM_LIST if build script present
+SET "PLATFORM=%~1"
+IF NOT DEFINED PLATFORM EXIT /B 0
+
+IF NOT EXIST "%ENGINE_PLATFORMS_DIR%\%PLATFORM%\Plugins\Experimental\NNERuntimeIREE\Source\ExternalBuild\Scripts\Windows\BuildWindows.bat" (
+	EXIT /B 1
+)
+
+IF NOT DEFINED PLATFORM_LIST (
+	SET "PLATFORM_LIST=%PLATFORM%"
+	EXIT /B 0
+)
+
+SET "PLATFORM_LIST=%PLATFORM_LIST% %PLATFORM%"
+EXIT /B 0
+
+
+:discover_all
+FOR /D %%D IN ("%ENGINE_PLATFORMS_DIR%\*") DO (
+	CALL :add_platform "%%~nxD"
+)
+EXIT /B 0
diff --git a/Engine/Plugins/Experimental/NNERuntimeIREE/Source/ExternalBuild/Source/Iree/Compiler/Plugins/Target/UnrealShader/CMakeLists.txt b/Engine/Plugins/Experimental/NNERuntimeIREE/Source/ExternalBuild/Source/Iree/Compiler/Plugins/Target/UnrealShader/CMakeLists.txt
new file mode 100644
--- /dev/null
+++ b/Engine/Plugins/Experimental/NNERuntimeIREE/Source/ExternalBuild/Source/Iree/Compiler/Plugins/Target/UnrealShader/CMakeLists.txt
@@ -0,0 +1,82 @@
+iree_add_all_subdirs()
+
+set(SPIRV_CROSS_ENABLE_MSL OFF CACHE BOOL "" FORCE)
+set(SPIRV_CROSS_ENABLE_GLSL ON CACHE BOOL "" FORCE) # Required to enable MSL
+set(SPIRV_CROSS_ENABLE_HLSL ON CACHE BOOL "" FORCE) # Required by UnrealShaders
+
+set(SPIRV_CROSS_EXCEPTIONS_TO_ASSERTIONS OFF CACHE BOOL "" FORCE)
+set(SPIRV_CROSS_CLI OFF CACHE BOOL "" FORCE)
+set(SPIRV_CROSS_ENABLE_TESTS OFF CACHE BOOL "" FORCE)
+set(SPIRV_CROSS_SKIP_INSTALL ON CACHE BOOL "" FORCE)
+
+set(SPIRV_CROSS_ENABLE_CPP OFF CACHE BOOL "" FORCE)
+set(SPIRV_CROSS_ENABLE_REFLECT OFF CACHE BOOL "" FORCE)
+set(SPIRV_CROSS_ENABLE_C_API OFF CACHE BOOL "" FORCE)
+set(SPIRV_CROSS_ENABLE_UTIL OFF CACHE BOOL "" FORCE)
+
+add_subdirectory(${CMAKE_SOURCE_DIR}/third_party/spirv_cross ${CMAKE_BINARY_DIR}/third_party/spirv_cross EXCLUDE_FROM_ALL)
+
+flatbuffer_c_library(
+  NAME
+    unreal_executable_def_c_fbs
+  SRCS
+    "unreal_executable_def.fbs"
+  FLATCC_ARGS
+    "--reader"
+    "--builder"
+    "--verifier"
+    "--json"
+  INCLUDES
+    "executable_debug_info.fbs"
+  PUBLIC
+)
+
+iree_cc_library(
+  NAME
+    SPIRVToHLSL
+  HDRS
+    "SPIRVToHLSL.h"
+  SRCS
+    "SPIRVToHLSL.cpp"
+  DEPS
+    LLVMSupport
+    MLIRSupport
+    spirv-cross-hlsl
+  PUBLIC
+)
+
+iree_cc_library(
+  NAME
+    UnrealShader
+  SRCS
+    "UnrealShaderTarget.cpp"
+    "UnrealShaderTargetUtils.cpp"
+  DEPS
+    ::SPIRVToHLSL
+    LLVMSupport
+    MLIRAsmParser
+    MLIRGPUDialect
+    MLIRIR
+    MLIRSPIRVDialect
+    MLIRSPIRVModuleCombiner
+    MLIRSPIRVSerialization
+    MLIRSupport
+    iree::compiler::Codegen::Common
+    iree::compiler::Codegen::Dialect::Codegen::IR::IREECodegenDialect
+    iree::compiler::Codegen::SPIRV
+    iree::compiler::Codegen::Utils
+    iree::compiler::Dialect::HAL::Target
+    iree::compiler::Dialect::HAL::Utils::ExecutableDebugInfoUtils
+    iree::compiler::PluginAPI
+    iree::compiler::Utils
+    iree::schemas::executable_debug_info_c_fbs
+    ::unreal_executable_def_c_fbs
+  PUBLIC
+)
+
+iree_compiler_register_plugin(
+  PLUGIN_ID
+    hal_target_unreal_shader
+  TARGET
+    ::UnrealShader
+)
\ No newline at end of file
diff --git a/Engine/Plugins/Experimental/NNERuntimeIREE/Source/ExternalBuild/Source/Iree/Compiler/Plugins/Target/UnrealShader/SPIRVToHLSL.cpp b/Engine/Plugins/Experimental/NNERuntimeIREE/Source/ExternalBuild/Source/Iree/Compiler/Plugins/Target/UnrealShader/SPIRVToHLSL.cpp
new file mode 100644
--- /dev/null
+++ b/Engine/Plugins/Experimental/NNERuntimeIREE/Source/ExternalBuild/Source/Iree/Compiler/Plugins/Target/UnrealShader/SPIRVToHLSL.cpp
@@ -0,0 +1,182 @@
+// Copyright Epic Games, Inc. All Rights Reserved.
+
+#include "SPIRVToHLSL.h"
+
+#include "llvm/ADT/STLExtras.h"
+#include "llvm/Support/Debug.h"
+#include "llvm/Support/raw_ostream.h"
+
+#include <set>
+#include <unordered_map>
+
+// Disable exception handling in favor of assertions.
+#define SPIRV_CROSS_EXCEPTIONS_TO_ASSERTIONS
+#include "third_party/spirv_cross/spirv_hlsl.hpp"
+
+#define DEBUG_TYPE "spirv-to-hlsl"
+
+namespace mlir::iree_compiler {
+
+namespace {
+
+bool GetCommonBaseType(SPIRV_CROSS_NAMESPACE::Compiler& compiler, const SPIRV_CROSS_NAMESPACE::SPIRType &type, SPIRV_CROSS_NAMESPACE::SPIRType::BaseType &baseType)
+{
+	if (type.basetype == SPIRV_CROSS_NAMESPACE::SPIRType::Struct)
+	{
+		baseType = SPIRV_CROSS_NAMESPACE::SPIRType::Unknown;
+		for (auto &member_type : type.member_types)
+		{
+			SPIRV_CROSS_NAMESPACE::SPIRType::BaseType member_base;
+			if (!GetCommonBaseType(compiler, compiler.get_type(member_type), member_base))
+			{
+				return false;
+			}
+
+			if (baseType == SPIRV_CROSS_NAMESPACE::SPIRType::Unknown)
+			{
+				baseType = member_base;
+			}
+			else if (baseType != member_base)
+			{
+				return false;
+			}
+		}
+		return true;
+	}
+	else
+	{
+		baseType = type.basetype;
+		return true;
+	}
+}
+
+}
+
+std::optional<std::pair<HLSLShader, std::string>> crossCompileSPIRVToHLSL(llvm::ArrayRef<uint32_t> spvBinary, StringRef entryPoint, bool useStructuredBuffers) {
+	SPIRV_CROSS_NAMESPACE::CompilerHLSL compiler(spvBinary.data(), spvBinary.size());
+
+	// Set Common options (GLSL)
+	{
+		SPIRV_CROSS_NAMESPACE::CompilerGLSL::Options Options = compiler.get_common_options();
+		// Options.emit_push_constant_as_uniform_buffer = true;
+
+		compiler.set_common_options(Options);
+	}
+
+	// Set HLSL specific options
+	{
+		SPIRV_CROSS_NAMESPACE::CompilerHLSL::Options hlslOptions = compiler.get_hlsl_options();
+		hlslOptions.shader_model = 62;
+		hlslOptions.enable_16bit_types = true;
+		hlslOptions.preserve_structured_buffers = useStructuredBuffers;
+
+		compiler.set_hlsl_options(hlslOptions);
+	}
+
+	// Do not emit any : register(#) bindings for ALL resource types, and rely on HLSL compiler to assign something.
+	compiler.set_resource_binding_flags(SPIRV_CROSS_NAMESPACE::HLSL_BINDING_AUTO_ALL);
+
+	// With HLSL... from spirv_cross/main.cpp
+	compiler.remap_num_workgroups_builtin();
+
+	SPIRV_CROSS_NAMESPACE::ShaderResources res = compiler.get_shader_resources();
+	if (res.uniform_buffers.size() > 0)
+	{
+		llvm::errs() << "CrossCompileSPIRVToHLSL: Uniform buffers not supported!\n";
+		return std::nullopt;
+	}
+
+	std::string shaderParameterMetadata = "Type;Name;ShaderType;Binding;DescriptorSet;ElementType;NumElements\n";
+
+	std::unordered_map<std::string, uint32_t> aliasCounts;
+
+	// Name buffers
+	for (auto& ssbo : res.storage_buffers)
+	{
+		const uint32_t binding = compiler.get_decoration(ssbo.id, spv::DecorationBinding);
+		const uint32_t set = compiler.get_decoration(ssbo.id, spv::DecorationDescriptorSet);
+
+		std::string name = "Buffer_" + std::to_string(set) + "_" + std::to_string(binding);
+
+		// Make sure names are unique for aliased buffers
+		if (compiler.get_decoration(ssbo.id, spv::DecorationAliased))
+		{
+			name += "_" + std::to_string(aliasCounts[name]++);
+		}
+
+		compiler.set_name(ssbo.id, name);
+
+		if (useStructuredBuffers)
+		{
+			compiler.set_decoration_string(ssbo.id, spv::DecorationUserTypeGOOGLE, "rwstructuredbuffer");
+
+			shaderParameterMetadata += "BUFFER_UAV;" + name + ";RWStructuredBuffer;" + std::to_string(binding) + ";" + std::to_string(set) + ";;\n";
+		}
+		else
+		{
+			shaderParameterMetadata += "BUFFER_UAV;" + name + ";RWByteAddressBuffer;" + std::to_string(binding) + ";" + std::to_string(set) + ";;\n";
+		}
+	}
+
+	if (res.push_constant_buffers.size() == 1)
+	{
+		// Push constant buffers have to be flattened as they could not be packed
+		compiler.flatten_buffer_block(res.push_constant_buffers[0].id);
+
+		const std::string name = "Constant";
+
+		// For flattened push constant buffers the base type name will be used
+		compiler.set_name(res.push_constant_buffers[0].base_type_id, name);
+
+		// Get array size and basic type, for spirv-cross implementation checkout: 
+		// void CompilerGLSL::emit_buffer_block_flattened(const SPIRVariable &var)
+		// in third_party\iree\third_party\spirv_cross\spirv_glsl.cpp
+		//
+		// The buffer will be flattened to array of vec4 with common basic type: float32, int32 or uint32.
+		const SPIRV_CROSS_NAMESPACE::SPIRType& type = compiler.get_type(res.push_constant_buffers[0].base_type_id);
+		const size_t vecSize = 4; // Check flatten code!
+		const size_t numVecElements = (compiler.get_declared_struct_size(type) + 15) / 16;
+
+		SPIRV_CROSS_NAMESPACE::SPIRType::BaseType baseType;
+		if (!GetCommonBaseType(compiler, type, baseType))
+		{
+			llvm::errs() << "CrossCompileSPIRVToHLSL: Push constant block contains multiple type, flatten will fail!\n";
+			return std::nullopt;
+		}
+		
+		std::string baseTypeStr;
+		if (baseType == SPIRV_CROSS_NAMESPACE::SPIRType::Float)
+		{
+			baseTypeStr = "FLOAT32";
+		}
+		else if (baseType == SPIRV_CROSS_NAMESPACE::SPIRType::Int)
+		{
+			baseTypeStr = "INT32";
+		}
+		else if (baseType == SPIRV_CROSS_NAMESPACE::SPIRType::UInt)
+		{
+			baseTypeStr = "UINT32";
+		}
+		else
+		{
+			llvm::errs() << "CrossCompileSPIRVToHLSL: Basic type " << baseType << " not supported!\n";
+			return std::nullopt;
+		}
+
+		shaderParameterMetadata += "PARAM_ARRAY;" + name + ";;;;" + baseTypeStr + ";" + std::to_string(numVecElements * vecSize) + "\n";
+	}
+	else if (res.push_constant_buffers.size() > 1)
+	{
+		// SPIR-V allows only one push constant block per shader stage, but just to be sure...
+		llvm::errs() << "CrossCompileSPIRVToHLSL: Multiple Push Constants not supported!\n";
+		return std::nullopt;
+	}
+
+	std::string hlslSource = compiler.compile();
+
+	const auto &spirvEntryPoint = compiler.get_entry_point(entryPoint.str(), spv::ExecutionModel::ExecutionModelGLCompute);
+	
+	return std::make_pair(HLSLShader{std::move(hlslSource), std::move(shaderParameterMetadata)}, spirvEntryPoint.name);
+}
+
+} // namespace mlir::iree_compiler
diff --git a/Engine/Plugins/Experimental/NNERuntimeIREE/Source/ExternalBuild/Source/Iree/Compiler/Plugins/Target/UnrealShader/SPIRVToHLSL.h b/Engine/Plugins/Experimental/NNERuntimeIREE/Source/ExternalBuild/Source/Iree/Compiler/Plugins/Target/UnrealShader/SPIRVToHLSL.h
new file mode 100644
--- /dev/null
+++ b/Engine/Plugins/Experimental/NNERuntimeIREE/Source/ExternalBuild/Source/Iree/Compiler/Plugins/Target/UnrealShader/SPIRVToHLSL.h
@@ -0,0 +1,25 @@
+// Copyright Epic Games, Inc. All Rights Reserved.
+
+#pragma once
+
+#include <array>
+#include <optional>
+#include <string>
+
+#include "llvm/ADT/ArrayRef.h"
+#include "llvm/ADT/StringRef.h"
+#include "mlir/Support/LLVM.h"
+
+namespace mlir::iree_compiler {
+
+struct HLSLShader {
+  std::string source;
+  std::string metadata;
+};
+
+// Cross compiles SPIR-V into HLSL source code for the
+// compute shader with |entryPoint| and returns the HLSL source and the new
+// entry point name. Returns std::nullopt on failure.
+std::optional<std::pair<HLSLShader, std::string>> crossCompileSPIRVToHLSL(llvm::ArrayRef<uint32_t> spvBinary, StringRef entryPoint, bool useStructuredBuffers = false);
+
+} // namespace mlir::iree_compiler
diff --git a/Engine/Plugins/Experimental/NNERuntimeIREE/Source/ExternalBuild/Source/Iree/Compiler/Plugins/Target/UnrealShader/UnrealShaderTarget.cpp b/Engine/Plugins/Experimental/NNERuntimeIREE/Source/ExternalBuild/Source/Iree/Compiler/Plugins/Target/UnrealShader/UnrealShaderTarget.cpp
new file mode 100644
--- /dev/null
+++ b/Engine/Plugins/Experimental/NNERuntimeIREE/Source/ExternalBuild/Source/Iree/Compiler/Plugins/Target/UnrealShader/UnrealShaderTarget.cpp
@@ -0,0 +1,300 @@
+// Copyright Epic Games, Inc. All Rights Reserved.
+
+#include "iree/compiler/Codegen/Dialect/Codegen/IR/IREECodegenDialect.h"
+#include "iree/compiler/Codegen/SPIRV/Passes.h"
+#include "iree/compiler/Dialect/HAL/Target/TargetRegistry.h"
+#include "iree/compiler/Dialect/HAL/Utils/ExecutableDebugInfoUtils.h"
+#include "iree/compiler/PluginAPI/Client.h"
+#include "iree/compiler/Utils/FlatbufferUtils.h"
+#include "iree/compiler/Utils/ModuleUtils.h"
+#include "llvm/ADT/STLExtras.h"
+#include "llvm/Support/CommandLine.h"
+#include "llvm/Support/raw_ostream.h"
+#include "mlir/AsmParser/AsmParser.h"
+#include "mlir/Dialect/GPU/IR/GPUDialect.h"
+#include "mlir/Dialect/SPIRV/IR/SPIRVAttributes.h"
+#include "mlir/Dialect/SPIRV/IR/SPIRVDialect.h"
+#include "mlir/Dialect/SPIRV/IR/SPIRVOps.h"
+#include "mlir/Dialect/SPIRV/IR/TargetAndABI.h"
+#include "mlir/Dialect/SPIRV/Linking/ModuleCombiner.h"
+#include "mlir/IR/Builders.h"
+#include "mlir/IR/BuiltinAttributes.h"
+#include "mlir/IR/BuiltinOps.h"
+#include "mlir/IR/DialectResourceBlobManager.h"
+#include "mlir/Target/SPIRV/Serialization.h"
+#include "SPIRVToHLSL.h"
+#include "unreal_executable_def_builder.h"
+#include "UnrealShaderTargetUtils.h"
+
+namespace mlir::iree_compiler::IREE::HAL {
+
+namespace {
+struct UnrealShaderTargetOptions {
+	// Use vp_android_baseline_2022 profile as the default target--it's a good
+	// lowest common denominator to guarantee the generated SPIR-V is widely
+	// accepted for now. Eventually we want to use a list for multi-targeting.
+	std::string target = "vp_android_baseline_2022";
+	bool useStructuredBuffers = false;
+	bool indirectBindings = false;
+
+	void bindOptions(OptionsBinder &binder) {
+		static llvm::cl::OptionCategory category("UnrealShader HAL Target");
+		binder.opt<std::string>(
+				"iree-unreal-target", target,
+				llvm::cl::desc(
+						"Unreal Shader target controlling the SPIR-V environment, which will be "
+						"cross-compiled to HLSL and fed to Unreal Engines Shader compiler."
+						"Currently the same schemes are supported as for the Vulkan target: 1) LLVM "
+						"CodeGen backend style: e.g., 'gfx*' for AMD GPUs and 'sm_*' for "
+						"NVIDIA GPUs; 2) architecture code name style: e.g., "
+						"'rdna3'/'valhall4'/'ampere'/'adreno' for AMD/ARM/NVIDIA/Qualcomm "
+						"GPUs; 3) product name style: 'rx7900xtx'/'rtx4090' for AMD/NVIDIA "
+						"GPUs. See "
+						"https://iree.dev/guides/deployment-configurations/gpu-vulkan/ for "
+						"more details.^"
+						"Note: this will likely change in the future!"));
+		binder.opt<bool>(
+				"iree-unreal-use-structured-buffers", useStructuredBuffers,
+				llvm::cl::desc(
+						"Force to use the (RW)StructuredBuffer type to be used for all SSBO's."));
+		binder.opt<bool>(
+				"iree-unreal-experimental-indirect-bindings", indirectBindings,
+				llvm::cl::desc(
+						"Force indirect bindings for all generated dispatches."));
+	}
+};
+} // namespace
+
+class UnrealTargetDevice : public TargetDevice {
+public:
+	UnrealTargetDevice(const UnrealShaderTargetOptions &options) : options_(options) {}
+
+	IREE::HAL::DeviceTargetAttr getDefaultDeviceTarget(MLIRContext *context, const TargetRegistry &targetRegistry) const override
+	{
+		Builder b(context);
+
+		SmallVector<NamedAttribute> deviceConfigAttrs;
+		auto deviceConfigAttr = b.getDictionaryAttr(deviceConfigAttrs);
+
+		SmallVector<NamedAttribute> executableConfigAttrs;
+		auto executableConfigAttr = b.getDictionaryAttr(executableConfigAttrs);
+
+		SmallVector<IREE::HAL::ExecutableTargetAttr> executableTargetAttrs;
+		targetRegistry.getTargetBackend("unreal-shader")->getDefaultExecutableTargets(context, "unreal", executableConfigAttr, executableTargetAttrs);
+
+		return IREE::HAL::DeviceTargetAttr::get(context, b.getStringAttr("unreal"), deviceConfigAttr, executableTargetAttrs);
+	}
+
+private:
+	const UnrealShaderTargetOptions &options_;
+};
+
+class UnrealShaderTargetBackend : public TargetBackend {
+public:
+	UnrealShaderTargetBackend(const UnrealShaderTargetOptions &options) : options_(options) {}
+
+	std::string getLegacyDefaultDeviceID() const override { return "unreal"; }
+
+	void getDefaultExecutableTargets(MLIRContext *context, StringRef deviceID, DictionaryAttr deviceConfigAttr, SmallVectorImpl<IREE::HAL::ExecutableTargetAttr> &executableTargetAttrs) const override {
+		executableTargetAttrs.push_back(getExecutableTarget(context, options_.indirectBindings));
+	}
+
+	IREE::HAL::ExecutableTargetAttr getExecutableTarget(MLIRContext *context, bool indirectBindings) const {
+		Builder b(context);
+		SmallVector<NamedAttribute> configItems;
+		auto addConfig = [&](StringRef name, Attribute value) {
+			configItems.emplace_back(b.getStringAttr(name), value);
+		};
+
+		if (auto target = GPU::UnrealShaderTargetUtils::getTargetDetails(options_.target, context)) {
+			addConfig("iree.gpu.target", target);
+		} else {
+			emitError(b.getUnknownLoc(), "Unknown Unreal target '")
+					<< options_.target << "'";
+			return nullptr;
+		}
+
+		return IREE::HAL::ExecutableTargetAttr::get(
+				context, b.getStringAttr("unreal-shader"),
+				indirectBindings ? b.getStringAttr("vulkan-spirv-fb-ptr")
+												 : b.getStringAttr("vulkan-spirv-fb"),
+				b.getDictionaryAttr(configItems));
+	}
+
+	void getDependentDialects(DialectRegistry &registry) const override {
+		registry.insert<IREE::Codegen::IREECodegenDialect, spirv::SPIRVDialect, gpu::GPUDialect, IREE::GPU::IREEGPUDialect>();
+	}
+
+	void buildConfigurationPassPipeline(IREE::HAL::ExecutableTargetAttr targetAttr, OpPassManager &passManager) override {
+		buildSPIRVCodegenConfigurationPassPipeline(passManager);
+	}
+
+	void buildTranslationPassPipeline(IREE::HAL::ExecutableTargetAttr targetAttr, OpPassManager &passManager) override {
+		buildSPIRVCodegenPassPipeline(passManager);
+	}
+
+	void buildLinkingPassPipeline(OpPassManager &passManager) override {
+		buildSPIRVLinkingPassPipeline(passManager);
+	}
+
+	LogicalResult serializeExecutable(const SerializationOptions &options, IREE::HAL::ExecutableVariantOp variantOp, OpBuilder &executableBuilder) override {
+		// Today we special-case external variants but in the future we could allow
+		// for a linking approach allowing both code generation and external .spv
+		// files to be combined together.
+		if (variantOp.isExternal()) {
+			return variantOp.emitError() << "unimplemented serializeExecutable for external .spv files";
+		}
+
+		ModuleOp innerModuleOp = variantOp.getInnerModule();
+		auto spirvModuleOps = innerModuleOp.getOps<spirv::ModuleOp>();
+		if (spirvModuleOps.empty()) {
+			return variantOp.emitError() << "should contain some spirv.module ops";
+		}
+
+		// Create a list of executable exports (by ordinal) to the SPIR-V module and
+		// entry point defining them.
+		auto unsortedExportOps = llvm::to_vector(variantOp.getOps<IREE::HAL::ExecutableExportOp>());
+		DenseMap<StringRef, std::tuple<IREE::HAL::ExecutableExportOp, uint64_t>> exportOrdinalMap;
+		for (auto exportOp : variantOp.getOps<IREE::HAL::ExecutableExportOp>()) {
+			uint64_t ordinal = 0;
+			if (std::optional<APInt> optionalOrdinal = exportOp.getOrdinal()) {
+				ordinal = optionalOrdinal->getZExtValue();
+			} else {
+				// For executables with only one entry point linking doesn't kick in at
+				// all. So the ordinal can be missing for this case.
+				if (!llvm::hasSingleElement(unsortedExportOps)) {
+					return exportOp.emitError() << "should have ordinal attribute";
+				}
+			}
+			exportOrdinalMap[exportOp.getSymName()] = std::make_tuple(exportOp, ordinal);
+		}
+		SmallVector<IREE::HAL::ExecutableExportOp> sortedExportOps;
+		sortedExportOps.resize(unsortedExportOps.size());
+		SmallVector<std::tuple<IREE::HAL::ExecutableExportOp, spirv::ModuleOp, spirv::EntryPointOp>> exportOps;
+		exportOps.resize(unsortedExportOps.size());
+		for (spirv::ModuleOp spirvModuleOp : spirvModuleOps) {
+			// Currently the spirv.module op should only have one entry point.
+			auto spirvEntryPoints = spirvModuleOp.getOps<spirv::EntryPointOp>();
+			if (!llvm::hasSingleElement(spirvEntryPoints)) {
+				// TODO: support multiple entry points. We only need them here to get
+				// the module name for dumping files.
+				return spirvModuleOp.emitError() << "expected to contain exactly one entry point";
+			}
+			spirv::EntryPointOp spirvEntryPointOp = *spirvEntryPoints.begin();
+			auto [exportOp, ordinal] = exportOrdinalMap.at(spirvEntryPointOp.getFn());
+			sortedExportOps[ordinal] = exportOp;
+			exportOps[ordinal] = std::make_tuple(exportOp, spirvModuleOp, spirvEntryPointOp);
+		}
+
+		FlatbufferBuilder builder;
+		iree_hal_unreal_ExecutableDef_start_as_root(builder);
+
+		// Generate optional per-export debug information.
+		// May be empty if no debug information was requested.
+		auto exportDebugInfos = createExportDefs(options.debugLevel, sortedExportOps, builder);
+
+		// Create a list of all serialized SPIR-V modules.
+		// TODO: unique the modules when each contains multiple entry points.
+		for (auto [exportOp, spirvModuleOp, spirvEntryPointOp] : exportOps) {
+			if (!options.dumpIntermediatesPath.empty()) {
+				std::string assembly;
+				llvm::raw_string_ostream os(assembly);
+				spirvModuleOp.print(os, OpPrintingFlags().useLocalScope());
+				dumpDataToPath(options.dumpIntermediatesPath, options.dumpBaseName, spirvEntryPointOp.getFn(), ".spirv.mlir", assembly);
+			}
+
+			// Serialize the spirv::ModuleOp into the binary blob.
+			SmallVector<uint32_t, 0> spirvBinary;
+			if (failed(spirv::serialize(spirvModuleOp, spirvBinary)) || spirvBinary.empty()) {
+				return spirvModuleOp.emitError() << "failed to serialize";
+			}
+
+			if (!options.dumpBinariesPath.empty()) {
+				dumpDataToPath<uint32_t>(options.dumpBinariesPath, options.dumpBaseName, spirvEntryPointOp.getFn(), ".spv", spirvBinary);
+			}
+
+
+			// We can use ArrayRef here given spvBinary reserves 0 bytes on stack.
+			ArrayRef spvData(spirvBinary.data(), spirvBinary.size());
+			std::optional<std::pair<HLSLShader, std::string>> HLSLResult = crossCompileSPIRVToHLSL(spvData, spirvEntryPointOp.getFn(), options_.useStructuredBuffers);
+			if (!HLSLResult) {
+				return spirvModuleOp.emitError() << "failed to cross compile SPIR-V to HLSL (" << llvm::join_items("_", options.dumpBaseName, spirvEntryPointOp.getFn()) << ")";
+			}
+
+			if (!options.dumpBinariesPath.empty()) {
+				dumpDataToPath(options.dumpBinariesPath, options.dumpBaseName, spirvEntryPointOp.getFn(), ".hlsl", HLSLResult->first.source);
+				dumpDataToPath(options.dumpBinariesPath, options.dumpBaseName, spirvEntryPointOp.getFn(), ".spmetadata", HLSLResult->first.metadata);
+			}
+		}
+
+		SmallVector<iree_hal_unreal_UnrealShaderDef_ref_t> unrealShaderRefs;
+		for (auto [exportOp, spirvModuleOp, spirvEntryPointOp] : exportOps) {
+			auto ordinalAttr = exportOp.getOrdinalAttr();
+			assert(ordinalAttr && "ordinals must be assigned");
+			int64_t ordinal = ordinalAttr.getInt();
+			
+			const std::string moduleName = options.dumpBaseName; // TODO 
+			const std::string entryPoint(spirvEntryPointOp.getFn());
+
+			// TODO currently followes the debug dump name schema
+			const std::string sourceFileName = llvm::join_items("_", moduleName, entryPoint);
+
+			auto sourceFileNameRef = builder.createString(sourceFileName);
+			auto moduleNameRef = builder.createString(moduleName);
+			auto entryPointRef = builder.createString(entryPoint);
+
+			iree_hal_unreal_UnrealShaderDef_start(builder);
+			iree_hal_unreal_UnrealShaderDef_source_file_name_add(builder, sourceFileNameRef);
+			iree_hal_unreal_UnrealShaderDef_module_name_add(builder, moduleNameRef);
+			iree_hal_unreal_UnrealShaderDef_entry_point_add(builder, entryPointRef);
+			iree_hal_unreal_UnrealShaderDef_debug_info_add(builder, exportDebugInfos[ordinal]);
+			unrealShaderRefs.push_back(iree_hal_unreal_UnrealShaderDef_end(builder));
+		}
+		auto unrealShadersRef = builder.createOffsetVecDestructive(unrealShaderRefs);
+
+		// Add top-level executable fields following their order of definition.
+		iree_hal_unreal_ExecutableDef_unreal_shaders_add(builder, unrealShadersRef);
+		iree_hal_unreal_ExecutableDef_end_as_root(builder);
+
+		// Add the binary data to the target executable.
+		auto binaryOp = executableBuilder.create<IREE::HAL::ExecutableBinaryOp>(
+				variantOp.getLoc(), variantOp.getSymName(),
+				variantOp.getTarget().getFormat(),
+				builder.getBufferAttr(executableBuilder.getContext()));
+		binaryOp.setMimeTypeAttr(executableBuilder.getStringAttr("application/x-flatbuffers"));
+
+		return success();
+	}
+
+private:
+	const UnrealShaderTargetOptions &options_;
+};
+
+namespace {
+
+struct UnrealShaderSession : public PluginSession<UnrealShaderSession, UnrealShaderTargetOptions, PluginActivationPolicy::DefaultActivated> {
+	void populateHALTargetDevices(IREE::HAL::TargetDeviceList &targets) {
+		// #hal.device.target<"unreal", ...
+		targets.add("unreal", [&]() {
+			return std::make_shared<UnrealTargetDevice>(options);
+		});
+	}
+
+	void populateHALTargetBackends(IREE::HAL::TargetBackendList &targets) {
+		// #hal.executable.target<"unreal-shader", ...
+		targets.add("unreal-shader", [&]() {
+			return std::make_shared<UnrealShaderTargetBackend>(options);
+		});
+	}
+};
+
+} // namespace
+
+} // namespace mlir::iree_compiler::IREE::HAL
+
+extern "C" bool iree_register_compiler_plugin_hal_target_unreal_shader(mlir::iree_compiler::PluginRegistrar *registrar) {
+	registrar->registerPlugin<mlir::iree_compiler::IREE::HAL::UnrealShaderSession>("hal_target_unreal_shader");
+	return true;
+}
+
+IREE_DEFINE_COMPILER_OPTION_FLAGS(mlir::iree_compiler::IREE::HAL::UnrealShaderTargetOptions);
diff --git a/Engine/Plugins/Experimental/NNERuntimeIREE/Source/ExternalBuild/Source/Iree/Compiler/Plugins/Target/UnrealShader/UnrealShaderTargetUtils.cpp b/Engine/Plugins/Experimental/NNERuntimeIREE/Source/ExternalBuild/Source/Iree/Compiler/Plugins/Target/UnrealShader/UnrealShaderTargetUtils.cpp
new file mode 100644
--- /dev/null
+++ b/Engine/Plugins/Experimental/NNERuntimeIREE/Source/ExternalBuild/Source/Iree/Compiler/Plugins/Target/UnrealShader/UnrealShaderTargetUtils.cpp
@@ -0,0 +1,39 @@
+// Copyright Epic Games, Inc. All Rights Reserved.
+
+#include "UnrealShaderTargetUtils.h"
+
+#include "iree/compiler/Codegen/Dialect/GPU/TargetUtils/KnownTargets.h"
+
+namespace mlir::iree_compiler::IREE::GPU::UnrealShaderTargetUtils
+{
+
+TargetAttr getTargetDetails(llvm::StringRef target, MLIRContext *context)
+{
+	const TargetAttr targetAttr = mlir::iree_compiler::IREE::GPU::getVulkanTargetDetails(target, context);
+	const ComputeBitwidths computeBits = targetAttr.getWgp().getCompute().getValue() & ~ComputeBitwidths::Int8;
+	const StorageBitwidths storageBits = targetAttr.getWgp().getStorage().getValue() & ~StorageBitwidths::B8;
+
+	// from Engine\Source\ThirdParty\ShaderConductor\ShaderConductor\External\DirectXShaderCompiler\include\dxc\DXIL\DxilConstants.h
+	const int32_t maxWorkgroupMemoryBytes = 8192 * 4; // const unsigned kMaxTGSMSize = 8192 * 4;
+
+	TargetWgpAttr wgp = TargetWgpAttr::get(context, 
+		ComputeBitwidthsAttr::get(context, computeBits), 
+		StorageBitwidthsAttr::get(context, storageBits), 
+		targetAttr.getWgp().getSubgroup(),
+		targetAttr.getWgp().getDot(),
+		targetAttr.getWgp().getMma(),
+		targetAttr.getWgp().getSubgroupSizeChoices(),
+		targetAttr.getWgp().getMaxWorkgroupSizes(),
+		targetAttr.getWgp().getMaxThreadCountPerWorkgroup(),
+		maxWorkgroupMemoryBytes,
+		targetAttr.getWgp().getMaxWorkgroupCounts(),
+		targetAttr.getWgp().getMaxLoadInstructionBits(),
+		targetAttr.getWgp().getSimdsPerWgp(),
+		targetAttr.getWgp().getVgprSpaceBits(),
+		targetAttr.getWgp().getExtra()
+	);
+
+	return TargetAttr::get(context, targetAttr.getArch(), targetAttr.getFeatures(), wgp, targetAttr.getChip());
+}
+
+} // namespace mlir::iree_compiler::IREE::GPU::UnrealShaderTargetUtils
diff --git a/Engine/Plugins/Experimental/NNERuntimeIREE/Source/ExternalBuild/Source/Iree/Compiler/Plugins/Target/UnrealShader/UnrealShaderTargetUtils.h b/Engine/Plugins/Experimental/NNERuntimeIREE/Source/ExternalBuild/Source/Iree/Compiler/Plugins/Target/UnrealShader/UnrealShaderTargetUtils.h
new file mode 100644
--- /dev/null
+++ b/Engine/Plugins/Experimental/NNERuntimeIREE/Source/ExternalBuild/Source/Iree/Compiler/Plugins/Target/UnrealShader/UnrealShaderTargetUtils.h
@@ -0,0 +1,16 @@
+// Copyright Epic Games, Inc. All Rights Reserved.
+
+#pragma once
+
+#include "iree/compiler/Codegen/Dialect/GPU/IR/IREEGPUAttrs.h"
+#include "llvm/ADT/StringRef.h"
+
+namespace mlir::iree_compiler::IREE::GPU::UnrealShaderTargetUtils {
+
+// Returns a TargetAttr to describe the details of the given |target|, which can
+// be a product name like "rtx3090"/"mali-g710"/"adreno" or an microarchitecture
+// name like "ampere"/"valhall". Returns a null TargetAttr if the given |target|
+// is not recognized.
+TargetAttr getTargetDetails(llvm::StringRef target, MLIRContext *context);
+
+} // namespace mlir::iree_compiler::IREE::GPU::UnrealShaderTargetUtils
diff --git a/Engine/Plugins/Experimental/NNERuntimeIREE/Source/ExternalBuild/Source/Iree/Compiler/Plugins/Target/UnrealShader/iree_compiler_plugin.cmake b/Engine/Plugins/Experimental/NNERuntimeIREE/Source/ExternalBuild/Source/Iree/Compiler/Plugins/Target/UnrealShader/iree_compiler_plugin.cmake
new file mode 100644
--- /dev/null
+++ b/Engine/Plugins/Experimental/NNERuntimeIREE/Source/ExternalBuild/Source/Iree/Compiler/Plugins/Target/UnrealShader/iree_compiler_plugin.cmake
@@ -0,0 +1 @@
+add_subdirectory(${CMAKE_CURRENT_LIST_DIR} target/UnrealShader)
\ No newline at end of file
diff --git a/Engine/Plugins/Experimental/NNERuntimeIREE/Source/ExternalBuild/Source/Iree/Runtime/CMakeLists.txt b/Engine/Plugins/Experimental/NNERuntimeIREE/Source/ExternalBuild/Source/Iree/Runtime/CMakeLists.txt
new file mode 100644
--- /dev/null
+++ b/Engine/Plugins/Experimental/NNERuntimeIREE/Source/ExternalBuild/Source/Iree/Runtime/CMakeLists.txt
@@ -0,0 +1,8 @@
+cmake_minimum_required(VERSION 3.21...3.24)
+project(IREE_RUNTIME_WINDOWS)
+
+# Prevent an issue of not correctly included modules
+include(CheckCSourceCompiles)
+
+# Include main module
+include(iree_runtime)
\ No newline at end of file
diff --git a/Engine/Plugins/Experimental/NNERuntimeIREE/Source/ExternalBuild/Source/NNEMlirTools/CMakeLists.txt b/Engine/Plugins/Experimental/NNERuntimeIREE/Source/ExternalBuild/Source/NNEMlirTools/CMakeLists.txt
new file mode 100644
--- /dev/null
+++ b/Engine/Plugins/Experimental/NNERuntimeIREE/Source/ExternalBuild/Source/NNEMlirTools/CMakeLists.txt
@@ -0,0 +1,63 @@
+cmake_minimum_required(VERSION 3.24...3.31)
+
+project(NNEMlirTools
+	VERSION 0.1
+	LANGUAGES CXX)
+
+set(CMAKE_CXX_STANDARD 20)
+set(CMAKE_CXX_STANDARD_REQUIRED ON)
+set(CMAKE_CXX_EXTENSIONS OFF)
+
+# ────────────────────────────────────────────────────────────────
+# IREE
+# ────────────────────────────────────────────────────────────────
+
+# IREE paths
+set(IREE_SOURCE_DIR ${UE_IREE_BUILD_ROOT}/iree)
+set(IREE_BINARY_DIR ${CMAKE_CURRENT_BINARY_DIR}/iree)
+
+# Add IREE
+add_subdirectory(${IREE_SOURCE_DIR} ${IREE_BINARY_DIR} EXCLUDE_FROM_ALL)
+
+# ────────────────────────────────────────────────────────────────
+# NNEMlirTools shared library
+# ────────────────────────────────────────────────────────────────
+
+add_library(NNEMlirTools SHARED)
+
+target_sources(NNEMlirTools
+	PUBLIC
+		FILE_SET HEADERS
+			BASE_DIRS
+				NNERuntimeIREE/Include
+			FILES
+				NNERuntimeIREE/Include/NNEMlirTools.h
+				NNERuntimeIREE/Include/NNEMlirTools_cxx_api.h
+	PRIVATE
+		NNERuntimeIREE/Source/NNEMlirTools.cpp)
+
+target_include_directories(NNEMlirTools PUBLIC
+	${CMAKE_CURRENT_SOURCE_DIR}/NNERuntimeIREE/Include)
+
+target_compile_definitions(NNEMlirTools PUBLIC
+	NNEMLIR_EXPORT)
+
+target_link_libraries(NNEMlirTools PRIVATE
+	ChloOps
+	iree::compiler::Dialect::Flow::IR
+	iree::compiler::Dialect::HAL::IR::HALDialect
+	iree::compiler::Dialect::LinalgExt::IR
+	iree::compiler::Dialect::Stream::IR
+	iree::compiler::Dialect::Util::IR
+	iree::compiler::plugins::input::Torch::torch-mlir-dialects::TMTensorDialectIR
+	iree::compiler::plugins::input::Torch::torch-mlir::TorchConversionDialectIR
+	iree::compiler::plugins::input::Torch::torch-mlir::TorchDialectIR
+	MLIRMLProgramDialect
+	MLIRParser
+	StablehloOps
+	VhloOps)
+
+if(APPLE)
+	# Apple ld duplicates libm when MLIR pulls both libm and libSystem.
+	target_link_options(NNEMlirTools PRIVATE LINKER:-no_warn_duplicate_libraries)
+endif()
\ No newline at end of file
diff --git a/Engine/Plugins/Experimental/NNERuntimeIREE/Source/ExternalBuild/Source/NNEMlirTools/NNERuntimeIREE/Include/NNEMlirTools.h b/Engine/Plugins/Experimental/NNERuntimeIREE/Source/ExternalBuild/Source/NNEMlirTools/NNERuntimeIREE/Include/NNEMlirTools.h
new file mode 100644
--- /dev/null
+++ b/Engine/Plugins/Experimental/NNERuntimeIREE/Source/ExternalBuild/Source/NNEMlirTools/NNERuntimeIREE/Include/NNEMlirTools.h
@@ -0,0 +1,97 @@
+// Copyright Epic Games, Inc. All Rights Reserved.
+
+// NNEMlirTools.h - minimal C API for introspecting MLIR modules via LLVM/MLIR.
+// The interface is provided as a single, versioned function table to ensure
+// binary compatibility across builds and convenient manual loading.
+// All objects are opaque handles that are released with the corresponding 'Release' functions.
+//
+// The API and its handles are not thread-safe.
+
+#pragma once
+
+#include <cstdint>
+#include <cstddef>
+
+#if defined(NNEMLIR_EXPORT)
+	#if defined(_MSC_VER)
+		#define NNEMLIR_API __declspec(dllexport)
+		#define NNEMLIR_CALL __cdecl
+	#elif defined(__GNUC__)
+		#define NNEMLIR_API __attribute__((visibility("default")))
+		#define NNEMLIR_CALL
+	#else
+		#error "Unsupported platform for NNEMLIR export"
+	#endif
+#else
+	#define NNEMLIR_API
+	#define NNEMLIR_CALL
+#endif
+
+#ifdef __cplusplus
+extern "C" {
+#endif
+
+/// Success / error codes returned through NNEMlirStatus.
+typedef enum {
+	NNEMLIR_SUCCESS          = 0,		///< Operation completed successfully.
+	NNEMLIR_ERR_PARSE_FAILED = 1,		///< Parsing of file/buffer failed.
+	NNEMLIR_ERR_BAD_ARGUMENT = 2,		///< Operation received invalid argument.
+	NNEMLIR_ERR_INTERNAL     = 0x7FFF	///< Unexpected internal failure.
+} NNEMlirStatusCode;
+
+/// Opaque handles
+struct NNEMlirStatus_s;   typedef struct NNEMlirStatus_s*   NNEMlirStatus;		///< Status result of an API call.
+struct NNEMlirContext_s;  typedef struct NNEMlirContext_s*  NNEMlirContext;		///< Context, contains MLIRContext with registered dialects.
+struct NNEMlirModule_s;   typedef struct NNEMlirModule_s*   NNEMlirModule;		///< Parsed MLIR module.
+struct NNEMlirFunction_s; typedef struct NNEMlirFunction_s* NNEMlirFunction;	///< (Public) function within a module.
+struct NNEMlirValue_s;    typedef struct NNEMlirValue_s*    NNEMlirValue;		///< Function argument or result.
+
+/// ABI version
+#define NNEMLIR_ABI_VERSION 1u
+
+/// Function table
+typedef struct NNEMlirApi {
+	uint32_t AbiVersion; ///< Must equal NNEMLIR_ABI_VERSION
+	uint32_t StructSize; ///< sizeof(NNEMlirApi)
+	
+	/* Status helpers */
+	const char*         (NNEMLIR_CALL *StatusToString)        (NNEMlirStatus);
+	NNEMlirStatusCode   (NNEMLIR_CALL *GetStatusCode)         (NNEMlirStatus);
+	
+	/* Context & parsing */
+	NNEMlirContext      (NNEMLIR_CALL *CreateContext)         (void);
+	NNEMlirStatus       (NNEMLIR_CALL *ParseModuleFromBuffer) (NNEMlirContext,const char* /* Buffer ptr (null-termination not required) */,size_t /* Buffer size */,NNEMlirModule*);
+	NNEMlirStatus       (NNEMLIR_CALL *ParseModuleFromFile)   (NNEMlirContext,const char* /* Path to MLIR file */,NNEMlirModule*);
+	
+	/* Module-level introspection */
+	size_t              (NNEMLIR_CALL *GetPublicFunctionCount)(NNEMlirModule);
+	NNEMlirFunction     (NNEMLIR_CALL *GetPublicFunction)     (NNEMlirModule,size_t);
+	
+	/* Function-level introspection */
+	const char*         (NNEMLIR_CALL *GetFunctionName)       (NNEMlirFunction);
+	size_t              (NNEMLIR_CALL *GetInputCount)         (NNEMlirFunction);
+	size_t              (NNEMLIR_CALL *GetResultCount)        (NNEMlirFunction);
+	NNEMlirValue        (NNEMLIR_CALL *GetInputValue)         (NNEMlirFunction,size_t);
+	NNEMlirValue        (NNEMLIR_CALL *GetResultValue)        (NNEMlirFunction,size_t);
+	
+	/* Value-level introspection */
+	const char*         (NNEMLIR_CALL *GetValueName)          (NNEMlirValue); ///< nullptr if unnamed.
+	const char*         (NNEMLIR_CALL *GetValueTypeText)      (NNEMlirValue); ///< E.g. "tensor<4x16xf32>".
+	const char*         (NNEMLIR_CALL *GetElementTypeText)    (NNEMlirValue); ///< Element type only, e.g. "f32".
+	size_t              (NNEMLIR_CALL *GetShape)              (NNEMlirValue,int64_t*,size_t);
+	
+	/* Release opaque handles */
+	void (NNEMLIR_CALL *ReleaseStatus)   (NNEMlirStatus);
+	void (NNEMLIR_CALL *ReleaseContext)  (NNEMlirContext);
+	void (NNEMLIR_CALL *ReleaseModule)   (NNEMlirModule);
+	void (NNEMLIR_CALL *ReleaseFunction) (NNEMlirFunction);
+	void (NNEMLIR_CALL *ReleaseValue)    (NNEMlirValue);
+} NNEMlirApi;
+
+/// Returns a pointer to the global NNEMlirApi or nullptr if the requested
+/// minimum ABI version is newer than the implementation.
+NNEMLIR_API const NNEMlirApi* NNEMLIR_CALL NNEMlirGetInterface(uint32_t MinVersion);
+
+#ifdef __cplusplus
+} /* extern "C" */
+#endif
\ No newline at end of file
diff --git a/Engine/Plugins/Experimental/NNERuntimeIREE/Source/ExternalBuild/Source/NNEMlirTools/NNERuntimeIREE/Include/NNEMlirTools_cxx_api.h b/Engine/Plugins/Experimental/NNERuntimeIREE/Source/ExternalBuild/Source/NNEMlirTools/NNERuntimeIREE/Include/NNEMlirTools_cxx_api.h
new file mode 100644
--- /dev/null
+++ b/Engine/Plugins/Experimental/NNERuntimeIREE/Source/ExternalBuild/Source/NNEMlirTools/NNERuntimeIREE/Include/NNEMlirTools_cxx_api.h
@@ -0,0 +1,399 @@
+// Copyright Epic Games, Inc. All Rights Reserved.
+
+// NNEMlirTools C++ API is a header only wrapper around the NNEMlirTools C API.
+//
+// The C++ API simplifies usage by automatically releasing resources in the destructors.
+
+#pragma once
+
+#include <cassert>
+#include <stdexcept>
+#include <string>
+#include <utility>
+#include <vector>
+
+#include "NNEMlirTools.h"
+
+#ifdef NNEMLIR_NO_EXCEPTIONS
+// The #ifndef is for the very special case where the user of this library wants to define their own way of handling errors.
+// NOTE: This header expects control flow to not continue after calling NNEMLIR_CXX_API_THROW
+#ifndef NNEMLIR_CXX_API_THROW
+#include <iostream>
+#define NNEMLIR_CXX_API_THROW(Message) \
+	do { \
+		std::cerr << Message << std::endl; \
+		abort(); \
+	} while (false)
+#endif
+#else
+#define NNEMLIR_CXX_API_THROW(Message) \
+	throw std::runtime_error(Message);
+#endif
+
+namespace NNEMlirTools {
+
+[[noreturn]] inline void ThrowError(const std::string& Message) {
+	NNEMLIR_CXX_API_THROW(Message);
+}
+
+// Api manages the NNEMlirApi object, which by design has to be manullay initialized.
+// We don't want to link againt the library.
+class Api {
+public:
+	static const Api& Instance() {
+
+		if (!ApiPtr)
+		{
+			ThrowError("Api::Initialize() has not been called");
+		}
+
+		static Api _Api;
+		return _Api;
+	}
+
+	static void Initialize(const NNEMlirApi* Ptr)
+	{
+		if (ApiPtr && ApiPtr != Ptr)
+		{
+			ThrowError("Api already initialised with a different table");
+		}
+
+		ApiPtr = Ptr;
+	}
+	
+	const NNEMlirApi* operator->() const { return ApiPtr; }
+
+private:
+	static const NNEMlirApi* ApiPtr;
+
+	Api() {
+
+	}
+
+	// Not copyable
+	Api(const Api&) = delete;
+	Api& operator=(const Api&) = delete;
+};
+
+inline const NNEMlirApi* Api::ApiPtr = nullptr;
+
+// Overloaded Release methods, not possible with C API
+inline void Release(NNEMlirStatus Sts) { Api::Instance()->ReleaseStatus(Sts); }
+inline void Release(NNEMlirContext Ctx) { Api::Instance()->ReleaseContext(Ctx); }
+inline void Release(NNEMlirModule Mod) { Api::Instance()->ReleaseModule(Mod); }
+inline void Release(NNEMlirFunction Func) { Api::Instance()->ReleaseFunction(Func); }
+inline void Release(NNEMlirValue Val) { Api::Instance()->ReleaseValue(Val); }
+
+// Owning pointer holder for the underlying C Objects that is move-only. Similar to std::unique_ptr.
+template<class T>
+class UniqueHandle
+{
+public:
+	UniqueHandle() noexcept = default;
+	explicit UniqueHandle(T Hdl) noexcept : Handle(Hdl) { }
+
+	~UniqueHandle() noexcept
+	{
+		Reset();
+	}
+
+	UniqueHandle(UniqueHandle&& Other) noexcept : Handle(std::exchange(Other.Handle, T{})) { }
+
+	UniqueHandle& operator=(UniqueHandle&& Other) noexcept
+	{
+		if (this != &Other)
+		{
+			Reset(Other.Release());
+		}
+		return *this;
+	}
+
+	void Reset(T NewHandle = T{}) noexcept
+	{
+		if (Handle)
+		{
+			NNEMlirTools::Release(Handle);
+		}
+
+		Handle = NewHandle;
+	}
+
+	T Release() noexcept
+	{
+		return std::exchange(Handle, T{});
+	}
+
+	// Move only
+	UniqueHandle(const UniqueHandle&) = delete;
+	UniqueHandle& operator=(const UniqueHandle&) = delete;
+
+	T Get() const noexcept { return Handle; }
+	explicit operator bool() const noexcept { return Handle != T{}; }
+
+private:
+	T Handle{};
+};
+
+class Function;
+class Value;
+
+// Status that is set on success/error and might contain a message.
+class Status : public UniqueHandle<NNEMlirStatus>
+{
+public:
+	Status(NNEMlirStatus Sts) noexcept : UniqueHandle<NNEMlirStatus>(Sts)
+	{
+		assert(Sts && "NNEMlirStatus handle must not be null");
+	}
+
+	// True on success or false on error
+	explicit operator bool() const
+	{
+		return Api::Instance()->GetStatusCode(Get()) == NNEMLIR_SUCCESS;
+	}
+
+	// Get message (might be only set on error)
+	std::string GetMessage() const
+	{
+		return Api::Instance()->StatusToString(Get());
+	}
+};
+
+// Context that holds the registered dialects
+class Context : public UniqueHandle<NNEMlirContext>
+{
+public:
+	Context() : UniqueHandle<NNEMlirContext>(Api::Instance()->CreateContext())
+	{
+		if (!Get())
+		{
+			ThrowError("Failed to create context.");
+		}
+	}
+};
+
+// Module obtained from parsing a MLIR file or buffer
+class Module : public UniqueHandle<NNEMlirModule>
+{
+public:
+	static Module ParseFromFile(const Context& Ctx, const char* Path)
+	{
+		NNEMlirModule Mod{};
+		const Status Sts = Api::Instance()->ParseModuleFromFile(Ctx.Get(), Path, &Mod);
+		if (!Sts)
+		{
+			ThrowError(Sts.GetMessage());
+		}
+		return Module(Mod);
+	}
+
+	// Parse MLIR from buffer (not null terminated) and return Module
+	static Module ParseFromBuffer(const Context& Ctx, const char* Data, size_t Length)
+	{
+		NNEMlirModule Mod{};
+		const Status Sts = Api::Instance()->ParseModuleFromBuffer(Ctx.Get(), Data, Length, &Mod);
+		if (!Sts)
+		{
+			ThrowError(Sts.GetMessage());
+		}
+		return Module(Mod);
+	}
+
+	// Return public function count
+	size_t GetFunctionCount() const
+	{
+		return Api::Instance()->GetPublicFunctionCount(Get());
+	}
+
+	// Retrieve a public function by index
+	Function GetFunction(size_t Index) const;
+
+	// Get a vector containing all public functions
+	std::vector<Function> GetFunctions() const;
+
+private:
+	Module(NNEMlirModule Mod) noexcept : UniqueHandle<NNEMlirModule>(Mod)
+	{
+		assert(Mod && "NNEMlirModule handle must not be null");
+	}
+};
+
+// Function in a module
+class Function : public UniqueHandle<NNEMlirFunction>
+{
+public:
+	// Get the name of the function
+	std::string GetName() const
+	{
+		const char* Name = Api::Instance()->GetFunctionName(Get());
+		return Name ? Name : "";
+	}
+
+	// Get the input argument count
+	size_t GetInputCount() const
+	{
+		return Api::Instance()->GetInputCount(Get());
+	}
+
+	// Get an input argument by index
+	Value GetInput(size_t Index) const;
+
+	// Get a vector with the input arguments
+	std::vector<Value> GetInputs() const;
+
+	// Get the result count
+	size_t GetResultCount() const
+	{
+		return Api::Instance()->GetResultCount(Get());
+	}
+
+	// Get a result by index
+	Value GetResult(size_t Index) const;
+
+	// Get a vector with the results
+	std::vector<Value> GetResults() const;
+
+private:
+	friend class Module;
+	explicit Function(NNEMlirFunction Fn) : UniqueHandle<NNEMlirFunction>(Fn)
+	{
+		assert(Fn && "NNEMlirFunction handle must not be null");
+	}
+};
+
+// Argument or result of a function
+class Value : public UniqueHandle<NNEMlirValue>
+{
+public:
+	// Get the element type string, e.g. "f32"
+	std::string GetElementType() const {
+		const char* Type = Api::Instance()->GetElementTypeText(Get());
+		return Type ? Type : "";
+	}
+
+	// Get the shape type text, e.g. "4x16xf32"
+	std::string GetShapeTypeText() const {
+		const std::vector<int64_t> Shape  = GetShape();
+		const std::string ElementType = GetElementType();
+
+		std::string Result;
+		for (size_t i = 0; i < Shape.size(); i++)
+		{
+			if (i) Result.push_back('x');
+			Result += (Shape[i] < 0 ? "?" : std::to_string(Shape[i]));
+		}
+		if (!Result.empty()) Result.push_back('x');
+		Result += ElementType;
+		
+		return Result;
+	}
+
+	// Get the name (empty for results)
+	std::string GetName() const {
+		const char* Name = Api::Instance()->GetValueName(Get());
+		return Name ? Name : "";
+	}
+
+	// Get the rank
+	size_t GetRank() const
+	{
+		return Api::Instance()->GetShape(Get(), nullptr, 0);
+	}
+
+	// Get the shape
+	std::vector<int64_t> GetShape() const {
+		const size_t Rank = GetRank();
+
+		std::vector<int64_t> Result(Rank);
+		if (Rank)
+		{
+			Api::Instance()->GetShape(Get(), Result.data(), Rank);
+		}
+		return Result;
+	}
+
+	// Fill in the shape
+	size_t GetShape(int64_t* ShapeArray, size_t Capacity) const {
+		return Api::Instance()->GetShape(Get(), ShapeArray, Capacity);
+	}
+
+private:
+	friend Function;
+	explicit Value(NNEMlirValue Val) : UniqueHandle<NNEMlirValue>(Val)
+	{
+		assert(Val && "NNEMlirValue handle must not be null");
+	}
+};
+
+inline Function Module::GetFunction(size_t Index) const
+{
+	if (Index >= GetFunctionCount())
+	{
+		ThrowError("GetFunction index out of range");
+	}
+	return Function(Api::Instance()->GetPublicFunction(Get(), Index));
+}
+
+inline std::vector<Function> Module::GetFunctions() const
+{
+	const size_t NumFuns = GetFunctionCount();
+
+	std::vector<Function> Result;
+	Result.reserve(NumFuns);
+	for (size_t i = 0; i < NumFuns; i++)
+	{
+		Result.push_back(GetFunction(i));
+	}
+
+	return Result;
+}
+
+inline Value Function::GetInput(size_t Index) const
+{
+	if (Index >= GetInputCount())
+	{
+		ThrowError("input index out of range");
+	}
+
+	return Value(Api::Instance()->GetInputValue(Get(), Index));
+}
+
+inline std::vector<Value> Function::GetInputs() const
+{
+	const size_t NumInputs = GetInputCount();
+
+	std::vector<Value> Result;
+	Result.reserve(NumInputs);
+	for (size_t i = 0; i < NumInputs; i++)
+	{
+		Result.push_back(GetInput(i));
+	}
+
+	return Result;
+}
+
+inline Value Function::GetResult(size_t Index) const
+{
+	if (Index >= GetResultCount())
+	{
+		ThrowError("result index out of range");
+	}
+
+	return Value(Api::Instance()->GetResultValue(Get(), Index));
+}
+
+inline std::vector<Value> Function::GetResults() const
+{
+	const size_t NumResults = GetResultCount();
+
+	std::vector<Value> Result;
+	Result.reserve(NumResults);
+	for (size_t i = 0; i < NumResults; i++)
+	{
+		Result.push_back(GetResult(i));
+	}
+
+	return Result;
+}
+
+} // namespace NNEMlirTools
diff --git a/Engine/Plugins/Experimental/NNERuntimeIREE/Source/ExternalBuild/Source/NNEMlirTools/NNERuntimeIREE/Source/NNEMlirTools.cpp b/Engine/Plugins/Experimental/NNERuntimeIREE/Source/ExternalBuild/Source/NNEMlirTools/NNERuntimeIREE/Source/NNEMlirTools.cpp
new file mode 100644
--- /dev/null
+++ b/Engine/Plugins/Experimental/NNERuntimeIREE/Source/ExternalBuild/Source/NNEMlirTools/NNERuntimeIREE/Source/NNEMlirTools.cpp
@@ -0,0 +1,524 @@
+// Copyright Epic Games, Inc. All Rights Reserved.
+
+#include "NNEMlirTools.h"
+
+#include "iree/compiler/Dialect/Flow/IR/FlowDialect.h"
+#include "iree/compiler/Dialect/HAL/IR/HALDialect.h"
+#include "iree/compiler/Dialect/LinalgExt/IR/LinalgExtDialect.h"
+#include "iree/compiler/Dialect/Stream/IR/StreamDialect.h"
+#include "iree/compiler/Dialect/Util/IR/UtilDialect.h"
+#include "llvm/ADT/StringRef.h"
+#include "llvm/Support/CommandLine.h"
+#include "llvm/Support/ErrorOr.h"
+#include "llvm/Support/MemoryBuffer.h"
+#include "llvm/Support/raw_ostream.h"
+#include "llvm/Support/SourceMgr.h"
+#include "mlir/Dialect/Arith/IR/Arith.h"
+#include "mlir/Dialect/Func/IR/FuncOps.h"
+#include "mlir/Dialect/Linalg/IR/Linalg.h"
+#include "mlir/Dialect/Math/IR/Math.h"
+#include "mlir/Dialect/MLProgram/IR/MLProgram.h"
+#include "mlir/Dialect/SCF/IR/SCF.h"
+#include "mlir/Dialect/Shape/IR/Shape.h"
+#include "mlir/Dialect/Tensor/IR/Tensor.h"
+#include "mlir/Dialect/Tosa/IR/TosaOps.h"
+#include "mlir/IR/BuiltinOps.h"
+#include "mlir/IR/MLIRContext.h"
+#include "mlir/Parser/Parser.h"
+#include "stablehlo/dialect/ChloOps.h"
+#include "stablehlo/dialect/StablehloOps.h"
+#include "stablehlo/dialect/VhloOps.h"
+#include "torch-mlir-dialects/Dialect/TMTensor/IR/TMTensorDialect.h"
+#include "torch-mlir/Dialect/Torch/IR/TorchDialect.h"
+#include "torch-mlir/Dialect/Torch/IR/TorchTypes.h"
+#include "torch-mlir/Dialect/TorchConversion/IR/TorchConversionDialect.h"
+
+//
+// Object hierarchy & lifetime:
+// * NNEMlirStatus: wrapper around status code & message.
+// * NNEMlirContext: owns shared pointer to mlir::MLIRContext. Multiple modules may share it.
+// * NNEMlirModule: owns shared pointer to mlir::ModuleOp and copy of context shared pointer to keep mlir::MLIRContext alive.
+// * NNEMlirFunction: non-owning view of mlir::func::FuncOp, but also has context and module shared pointers.
+// * NNEMlirValue: non-owning view of BlockArgument (input) or Type (result). Also has context and module shared pointers.
+//
+// All handles are created with 'new' and released with matching 'NNEMlirReleaseXYZ()'. All handles are safe to destroy in any order
+// as every handle keeps a shared reference to the context and module it was created from.
+//
+// Returned c-string pointers have the same lifetime as the object that returned it.
+
+struct NNEMlirStatus_s
+{
+	NNEMlirStatusCode Code;
+	std::string Message;
+
+	explicit NNEMlirStatus_s(NNEMlirStatusCode C) : Code(C) { }
+
+	NNEMlirStatus_s(NNEMlirStatusCode C, const std::string& Msg) : Code(C), Message(Msg) { }
+};
+
+// Create NNEMlirStatus helper macros
+#define MAKE_STATUS_SUCCESS() \
+	new NNEMlirStatus_s(NNEMLIR_SUCCESS);
+
+#define MAKE_STATUS_FROM_CODE(Code) \
+	new NNEMlirStatus_s(Code);
+
+#define MAKE_ERROR_STATUS(Code, Message) \
+	new NNEMlirStatus_s(Code, Message);
+
+struct NNEMlirContext_s
+{
+	std::shared_ptr<mlir::MLIRContext> Context;
+	
+	explicit NNEMlirContext_s(std::shared_ptr<mlir::MLIRContext> Ctx) : Context(std::move(Ctx)) { }
+};
+
+struct NNEMlirModule_s
+{
+	// Keep mlir context and module alive
+	std::shared_ptr<mlir::MLIRContext> Context;
+	std::shared_ptr<mlir::OwningOpRef<mlir::ModuleOp>> Module;
+
+	std::vector<mlir::func::FuncOp> PublicFuncs;
+
+	NNEMlirModule_s(std::shared_ptr<mlir::MLIRContext> Ctx, std::shared_ptr<mlir::OwningOpRef<mlir::ModuleOp>> Mod)
+		: Context(std::move(Ctx)), Module(std::move(Mod))
+	{
+		PublicFuncs.reserve(16);
+		for (mlir::func::FuncOp Func : Module->get().getOps<mlir::func::FuncOp>())
+		{
+			if (Func.isPublic())
+			{
+				PublicFuncs.push_back(Func);
+			}
+		}
+	}
+};
+
+struct NNEMlirFunction_s
+{
+	// Keep mlir context and module alive
+	std::shared_ptr<mlir::MLIRContext> Context;
+	std::shared_ptr<mlir::OwningOpRef<mlir::ModuleOp>> Module;
+
+	mlir::func::FuncOp Func;
+
+	std::string CachedName;
+
+	NNEMlirFunction_s(std::shared_ptr<mlir::MLIRContext> Ctx, std::shared_ptr<mlir::OwningOpRef<mlir::ModuleOp>> Mod, mlir::func::FuncOp Fn)
+		: Context(Ctx), Module(std::move(Mod)), Func(Fn) { }
+};
+
+struct NNEMlirValue_s
+{
+	// Keep mlir context and module alive
+	std::shared_ptr<mlir::MLIRContext> Context;
+	std::shared_ptr<mlir::OwningOpRef<mlir::ModuleOp>> Module;
+
+	mlir::Value Argument;
+	mlir::Type ResultType;
+	bool bIsResult; // Either 'Argument' or 'ResultType' is used
+
+	std::string CachedName;
+	std::string CachedType;
+	std::string CachedElementType;
+
+	NNEMlirValue_s(std::shared_ptr<mlir::MLIRContext> Ctx, std::shared_ptr<mlir::OwningOpRef<mlir::ModuleOp>> Mod, mlir::Value Arg)
+		: Context(Ctx), Module(std::move(Mod)), Argument(Arg), bIsResult(false) { }
+
+	NNEMlirValue_s(std::shared_ptr<mlir::MLIRContext> Ctx, std::shared_ptr<mlir::OwningOpRef<mlir::ModuleOp>> Mod, mlir::Type ResTy)
+		: Context(Ctx), Module(std::move(Mod)), ResultType(ResTy), bIsResult(true) { }
+};
+
+const char* NNEMlirStatusToString(NNEMlirStatus Status)
+{
+	if (!Status->Message.empty())
+	{
+		return Status->Message.c_str();
+	}
+
+	switch(Status->Code)
+	{
+		case NNEMLIR_SUCCESS:			return "success";
+		case NNEMLIR_ERR_PARSE_FAILED:	return "failed to parse";
+		case NNEMLIR_ERR_BAD_ARGUMENT:	return "bad argument";
+		case NNEMLIR_ERR_INTERNAL:		return "internal error";
+		default:						return "unknown status";
+	}
+}
+
+NNEMlirStatusCode NNEMlirGetStatusCode(NNEMlirStatus Status)
+{
+	return Status->Code;
+}
+
+NNEMlirContext NNEMlirCreateContext()
+{
+	std::shared_ptr<mlir::MLIRContext> Context = std::make_shared<mlir::MLIRContext>();
+
+	// Register Built-in dialects, selected according IREE docs:
+	// https://iree.dev/reference/extensions/#1-target-iree-input-dialects
+	Context->getOrLoadDialect<mlir::arith::ArithDialect>();
+	Context->getOrLoadDialect<mlir::func::FuncDialect>();
+	Context->getOrLoadDialect<mlir::linalg::LinalgDialect>();
+	Context->getOrLoadDialect<mlir::math::MathDialect>();
+	Context->getOrLoadDialect<mlir::scf::SCFDialect>();
+	Context->getOrLoadDialect<mlir::tensor::TensorDialect>();
+	// ...add more dialects?
+
+	// Register IREE's three main input dialects:
+	//
+	// 1. StableHLO
+	Context->getOrLoadDialect<mlir::chlo::ChloDialect>();
+	Context->getOrLoadDialect<mlir::shape::ShapeDialect>();
+	Context->getOrLoadDialect<mlir::stablehlo::StablehloDialect>();
+	Context->getOrLoadDialect<mlir::vhlo::VhloDialect>();
+	// 2. Torch-MLIR
+	Context->getOrLoadDialect<mlir::iree_compiler::IREE::LinalgExt::IREELinalgExtDialect>();
+	Context->getOrLoadDialect<mlir::ml_program::MLProgramDialect>();
+	Context->getOrLoadDialect<mlir::torch::TMTensor::TMTensorDialect>();
+	Context->getOrLoadDialect<mlir::torch::Torch::TorchDialect>();
+	Context->getOrLoadDialect<mlir::torch::TorchConversion::TorchConversionDialect>();
+	// 3. TOSA
+	Context->getOrLoadDialect<mlir::tosa::TosaDialect>();
+
+	// Register IREE internal dialect(s)
+	// Specifically because torch might convert into them as input to the IREE compilation pipeline:
+	// iree-org\iree\compiler\plugins\input\Torch\PluginRegistration.cpp
+	Context->getOrLoadDialect<mlir::iree_compiler::IREE::Flow::FlowDialect>();
+	Context->getOrLoadDialect<mlir::iree_compiler::IREE::HAL::HALDialect>();
+	Context->getOrLoadDialect<mlir::iree_compiler::IREE::Stream::StreamDialect>();
+	Context->getOrLoadDialect<mlir::iree_compiler::IREE::Util::UtilDialect>();
+
+	return new NNEMlirContext_s(std::move(Context));
+}
+
+void NNEMlirReleaseStatus(NNEMlirStatus Status)
+{
+	if (!Status) return;
+	delete Status;
+}
+
+void NNEMlirReleaseContext(NNEMlirContext Context)
+{
+	if (!Context) return;
+	delete Context;
+}
+
+void NNEMlirReleaseModule(NNEMlirModule Module)
+{
+	if (!Module) return;
+	delete Module;
+}
+
+void NNEMlirReleaseFunction(NNEMlirFunction Func)
+{
+	if (!Func) return;
+	delete Func;
+}
+
+void NNEMlirReleaseValue(NNEMlirValue Value)
+{
+	if (!Value) return;
+	delete Value;
+}
+
+NNEMlirStatus ParseSourceFile(NNEMlirContext Context, std::unique_ptr<llvm::MemoryBuffer> MemoryBuffer, NNEMlirModule* OutModule)
+{
+	llvm::SourceMgr SourceMgr;
+	SourceMgr.AddNewSourceBuffer(std::move(MemoryBuffer), llvm::SMLoc());
+
+	std::string DiagnosticsStr;
+	llvm::raw_string_ostream DiagOS(DiagnosticsStr);
+	mlir::SourceMgrDiagnosticHandler SourceMgrHandler(SourceMgr, Context->Context.get(), DiagOS);
+
+	const mlir::ParserConfig Config(Context->Context.get(), false);
+	mlir::OwningOpRef<mlir::ModuleOp> ModuleRef = mlir::parseSourceFile<mlir::ModuleOp>(SourceMgr, Config);
+	if (!ModuleRef)
+	{
+		DiagOS.flush();
+
+		return MAKE_ERROR_STATUS(NNEMLIR_ERR_PARSE_FAILED, DiagnosticsStr.empty() ? "Failed to parse from buffer" : DiagnosticsStr);
+	}
+
+	*OutModule = new NNEMlirModule_s(Context->Context, std::make_shared<mlir::OwningOpRef<mlir::ModuleOp>>(std::move(ModuleRef)));
+
+	return MAKE_STATUS_SUCCESS();
+}
+
+NNEMlirStatus NNEMlirParseModuleFromBuffer(NNEMlirContext Context, const char* Data, size_t Length, NNEMlirModule* OutModule)
+{
+	if (!Context || !Data || Length == 0 || !OutModule)
+	{
+		return MAKE_STATUS_FROM_CODE(NNEMLIR_ERR_BAD_ARGUMENT);
+	}
+
+	llvm::StringRef Source(Data, Length);
+
+	// IREE 3.4.0, note: we need to copy and let add a traling NUL. Can be removed if LLVM/MLIR fixed that.
+	llvm::ErrorOr<std::unique_ptr<llvm::MemoryBuffer>> FileOrErr = llvm::MemoryBuffer::getMemBufferCopy(Source, "<Buffer>");
+	if (std::error_code ErrorCode = FileOrErr.getError())
+	{
+		return MAKE_ERROR_STATUS(NNEMLIR_ERR_INTERNAL, "Could not get memory buffer copy: " + ErrorCode.message());
+	}
+
+	return ParseSourceFile(Context, std::move(*FileOrErr), OutModule);
+}
+
+NNEMlirStatus NNEMlirParseModuleFromFile(NNEMlirContext Context, const char* Path, NNEMlirModule* OutModule)
+{
+	if (!Context || !Path || !OutModule)
+	{
+		return MAKE_STATUS_FROM_CODE(NNEMLIR_ERR_BAD_ARGUMENT);
+	}
+
+	llvm::ErrorOr<std::unique_ptr<llvm::MemoryBuffer>> FileOrErr = llvm::MemoryBuffer::getFileOrSTDIN(Path);
+	if (std::error_code ErrorCode = FileOrErr.getError())
+	{
+		return MAKE_ERROR_STATUS(NNEMLIR_ERR_INTERNAL, "Could not open input file: " + ErrorCode.message());
+	}
+
+	return ParseSourceFile(Context, std::move(*FileOrErr), OutModule);
+}
+
+size_t NNEMlirGetPublicFunctionCount(NNEMlirModule Module)
+{
+	if (!Module) return 0;
+
+	return Module->PublicFuncs.size();
+}
+
+NNEMlirFunction NNEMlirGetPublicFunction(NNEMlirModule Module, size_t Index)
+{
+	if (!Module) return nullptr;
+	if (Index >= Module->PublicFuncs.size()) return nullptr;
+
+	const mlir::func::FuncOp Func = Module->PublicFuncs[Index];
+
+	return new NNEMlirFunction_s(Module->Context, Module->Module, Func);
+}
+
+const char* NNEMlirGetFunctionName(NNEMlirFunction Func)
+{
+	if (!Func) return nullptr;
+
+	if (Func->CachedName.empty())
+	{
+		Func->CachedName = Func->Func.getSymName().str();
+	}
+
+	return Func->CachedName.c_str();
+}
+
+size_t NNEMlirGetInputCount(NNEMlirFunction Func)
+{
+	if (!Func) return 0;
+
+	return Func->Func.getNumArguments();
+}
+
+size_t NNEMlirGetResultCount(NNEMlirFunction Func)
+{
+	if (!Func) return 0;
+
+	return Func->Func.getNumResults();
+}
+
+NNEMlirValue NNEMlirGetInputValue(NNEMlirFunction Func, size_t Index)
+{
+	if (!Func) return nullptr;
+	if (Index >= Func->Func.getNumArguments()) return nullptr;
+
+	const mlir::Value Argument = Func->Func.getArgument(Index);
+
+	return new NNEMlirValue_s(Func->Context, Func->Module, Argument);
+}
+
+NNEMlirValue NNEMlirGetResultValue(NNEMlirFunction Func, size_t Index)
+{
+	if (!Func) return nullptr;
+	if (Index >= Func->Func.getNumResults()) return nullptr;
+
+	const mlir::Type ResultType = Func->Func.getFunctionType().getResult(Index);
+
+	return new NNEMlirValue_s(Func->Context, Func->Module, ResultType);
+}
+
+const char* NNEMlirGetValueName(NNEMlirValue Value)
+{
+	if (!Value) return nullptr;
+	if (Value->bIsResult) return nullptr;
+
+	if (Value->CachedName.empty())
+	{
+		const mlir::BlockArgument BlockArg = mlir::cast<mlir::BlockArgument>(Value->Argument);
+		mlir::func::FuncOp Func = mlir::cast<mlir::func::FuncOp>(BlockArg.getOwner()->getParentOp());
+
+		const unsigned int Index = BlockArg.getArgNumber();
+		if (auto Attr = Func.getArgAttrOfType<mlir::StringAttr>(Index, "llvm.name"))
+		{
+			Value->CachedName = Attr.getValue().str();
+		}
+	}
+
+	return Value->CachedName.empty() ? nullptr : Value->CachedName.c_str();
+}
+
+const char* NNEMlirGetValueTypeText(NNEMlirValue Value)
+{
+	if (!Value) return nullptr;
+
+	if (Value->CachedType.empty())
+	{
+		mlir::Type Type = Value->bIsResult ? Value->ResultType : Value->Argument.getType();
+
+		std::string TypeStr;
+		llvm::raw_string_ostream Os(TypeStr);
+		Type.print(Os);
+
+		Value->CachedType = TypeStr;
+	}
+
+	return Value->CachedType.c_str();
+}
+
+const char* NNEMlirGetElementTypeText(NNEMlirValue Value)
+{
+	if (!Value) return nullptr;
+
+	if (Value->CachedElementType.empty())
+	{
+		mlir::Type Type = Value->bIsResult ? Value->ResultType : Value->Argument.getType();
+
+		// 1. Built-in tensors / memrefs etc.
+		if (auto Shaped = mlir::dyn_cast<mlir::ShapedType>(Type))
+		{
+			Type = Shaped.getElementType();
+		} // 2. Torch-MLIR tensors
+		else if (auto TorchTensor = mlir::dyn_cast<mlir::torch::Torch::BaseTensorType>(Type))
+		{
+			if (auto DType = TorchTensor.getOptionalDtype())
+			{
+				Type = DType;
+			}
+		}
+		// else{} 3. Something else
+
+		std::string TypeStr;
+		llvm::raw_string_ostream Os(TypeStr);
+		Type.print(Os);
+
+		Value->CachedElementType = TypeStr;
+	}
+
+	return Value->CachedElementType.c_str();
+}
+
+size_t NNEMlirGetShape(NNEMlirValue Value, int64_t* Dims, size_t Capacity)
+{
+	if (!Value)
+	{
+		return 0;
+	}
+
+	mlir::Type Type = Value->bIsResult ? Value->ResultType : Value->Argument.getType();
+
+	// 1. Built-in tensors / memrefs etc.
+	if (auto Shaped = mlir::dyn_cast<mlir::ShapedType>(Type))
+	{
+		if (!Shaped.hasRank())
+		{
+			return 0;
+		}
+
+		const size_t Rank = Shaped.getRank();
+		if (!Dims || Capacity == 0)
+		{
+			return Rank;
+		}
+
+		const size_t Num = std::min(Rank, Capacity);
+		const llvm::ArrayRef<int64_t> Shape = Shaped.getShape();
+		for (size_t i = 0; i < Num; ++i)
+		{
+			int64_t Dim = Shape[i];
+			if (mlir::ShapedType::isDynamic(Dim))
+			{
+				Dim = -1;
+			}
+			Dims[i] = Dim;
+		}
+
+		return Num;
+	}
+
+	// 2. Torch-MLIR tensors
+	if (auto TorchTensor = mlir::dyn_cast<mlir::torch::Torch::BaseTensorType>(Type))
+	{
+		if (auto OptSizes = TorchTensor.getOptionalSizes())
+		{
+			const llvm::ArrayRef<int64_t> Shape = *OptSizes;
+			const size_t Rank = Shape.size();
+
+			if (!Dims || Capacity == 0)
+			{
+				return Rank;
+			}
+
+			const size_t Num = std::min(Rank, Capacity);
+			for (size_t i = 0; i < Num; ++i)
+			{
+				int64_t Dim = Shape[i];
+				if (Dim == mlir::torch::Torch::kUnknownSize)
+				{
+					Dim = -1;
+				}
+				Dims[i] = Dim;
+			}
+
+			return Num;
+		}
+
+		// unranked !torch.tensor
+		return 0;
+	}
+
+	// 3. Something else
+	return 0;
+}
+
+static const NNEMlirApi G_Api = {
+	NNEMLIR_ABI_VERSION,
+	sizeof(NNEMlirApi),
+	&NNEMlirStatusToString,
+	&NNEMlirGetStatusCode,
+	&NNEMlirCreateContext,
+	&NNEMlirParseModuleFromBuffer,
+	&NNEMlirParseModuleFromFile,
+	&NNEMlirGetPublicFunctionCount,
+	&NNEMlirGetPublicFunction,
+	&NNEMlirGetFunctionName,
+	&NNEMlirGetInputCount,
+	&NNEMlirGetResultCount,
+	&NNEMlirGetInputValue,
+	&NNEMlirGetResultValue,
+	&NNEMlirGetValueName,
+	&NNEMlirGetValueTypeText,
+	&NNEMlirGetElementTypeText,
+	&NNEMlirGetShape,
+	&NNEMlirReleaseStatus,
+	&NNEMlirReleaseContext,
+	&NNEMlirReleaseModule,
+	&NNEMlirReleaseFunction,
+	&NNEMlirReleaseValue
+};
+
+extern "C"
+{
+
+const NNEMlirApi* NNEMlirGetInterface(uint32_t MinVersion)
+{
+	return (MinVersion <= G_Api.AbiVersion) ? &G_Api : nullptr;
+}
+
+}
\ No newline at end of file
