[NNE] NNERuntimeIREE minor fixes to build scripts.

#tests Ran scripts

#rnx

diff --git a/Engine/Plugins/Experimental/NNERuntimeIREE/Source/ExternalBuild/Scripts/Mac/BuildAndInstallMac.sh b/Engine/Plugins/Experimental/NNERuntimeIREE/Source/ExternalBuild/Scripts/Mac/BuildAndInstallMac.sh
--- a/Engine/Plugins/Experimental/NNERuntimeIREE/Source/ExternalBuild/Scripts/Mac/BuildAndInstallMac.sh
+++ b/Engine/Plugins/Experimental/NNERuntimeIREE/Source/ExternalBuild/Scripts/Mac/BuildAndInstallMac.sh
@@ -10,7 +10,7 @@ set -Eeuo pipefail
 
 trap 'echo; echo "Script failed. Exiting" >&2; exit 1' ERR
 
-SCRIPT_DIR="$(cd -- "$(dirname -- "${BASH_SOURCE[0]}")" &>/dev/null && pwd)"
+SCRIPT_DIR="$(cd -- "$(dirname "${BASH_SOURCE[0]}")" >/dev/null 2>&1 && pwd)"
 
 # For production be aware that the build path is part of the IREE error messages
 # So choose it with care
diff --git a/Engine/Plugins/Experimental/NNERuntimeIREE/Source/ExternalBuild/Scripts/Mac/BuildMac.sh b/Engine/Plugins/Experimental/NNERuntimeIREE/Source/ExternalBuild/Scripts/Mac/BuildMac.sh
--- a/Engine/Plugins/Experimental/NNERuntimeIREE/Source/ExternalBuild/Scripts/Mac/BuildMac.sh
+++ b/Engine/Plugins/Experimental/NNERuntimeIREE/Source/ExternalBuild/Scripts/Mac/BuildMac.sh
@@ -16,7 +16,7 @@ fi
 WORKING_DIR=$1
 
 BUILD_DIR=$WORKING_DIR/iree-org
-BUILD_SCRIPT_DIR="$(dirname -- "$(realpath -- "$0")")/"
+BUILD_SCRIPT_DIR="$(cd -- "$(dirname "${BASH_SOURCE[0]}")" >/dev/null 2>&1 && pwd)"
 UE_SOURCE_DIR=$BUILD_SCRIPT_DIR/../../Source
 UE_PATCHES_DIR=$BUILD_SCRIPT_DIR/../../Patch
 UE_CMAKE_DIR=$BUILD_SCRIPT_DIR/../../CMake
@@ -121,7 +121,7 @@ if [ ! -d "$BUILD_DIR/iree-compiler" ]; then
 		-DIREE_TARGET_BACKEND_DEFAULTS=OFF \
 		-DIREE_TARGET_BACKEND_LLVM_CPU=ON \
 		-DIREE_ERROR_ON_MISSING_SUBMODULES=OFF \
-		-DCMAKE_OSX_ARCHITECTURES=$ARCHITECTURES \
+		-DCMAKE_OSX_ARCHITECTURES="$ARCHITECTURES" \
 		-DCMAKE_OSX_DEPLOYMENT_TARGET=$OSX_TARGET \
 		-DIREE_EMBEDDED_RELEASE_INFO=ON \
 		-DIREE_RELEASE_VERSION=$IREE_COMPILER_VERSION_STRING \
@@ -161,7 +161,7 @@ if [ ! -d "$BUILD_DIR/NNEMlirTools" ]; then
 		-DIREE_TARGET_BACKEND_DEFAULTS=OFF \
 		-DIREE_TARGET_BACKEND_LLVM_CPU=OFF \
 		-DIREE_ERROR_ON_MISSING_SUBMODULES=OFF \
-		-DCMAKE_OSX_ARCHITECTURES=$ARCHITECTURES \
+		-DCMAKE_OSX_ARCHITECTURES="$ARCHITECTURES" \
 		-DCMAKE_OSX_DEPLOYMENT_TARGET=$OSX_TARGET
 		
 	cmake --build $BUILD_DIR/NNEMlirTools
@@ -203,7 +203,7 @@ if [ ! -d "$BUILD_DIR/iree-runtime-mac" ]; then
 		-DIREE_ENABLE_THREADING=ON \
 		-DIREE_HAL_EXECUTABLE_LOADER_DEFAULTS=OFF \
 		-DIREE_HAL_EXECUTABLE_PLUGIN_DEFAULTS=OFF \
-		-DCMAKE_OSX_ARCHITECTURES=$ARCHITECTURES \
+		-DCMAKE_OSX_ARCHITECTURES="$ARCHITECTURES" \
 		-DCMAKE_OSX_DEPLOYMENT_TARGET=$OSX_TARGET
 		
 	cmake --build $BUILD_DIR/iree-runtime-mac
diff --git a/Engine/Plugins/Experimental/NNERuntimeIREE/Source/ExternalBuild/Scripts/Mac/InstallMac.sh b/Engine/Plugins/Experimental/NNERuntimeIREE/Source/ExternalBuild/Scripts/Mac/InstallMac.sh
--- a/Engine/Plugins/Experimental/NNERuntimeIREE/Source/ExternalBuild/Scripts/Mac/InstallMac.sh
+++ b/Engine/Plugins/Experimental/NNERuntimeIREE/Source/ExternalBuild/Scripts/Mac/InstallMac.sh
@@ -16,7 +16,7 @@ fi
 WORKING_DIR=$1
 
 BUILD_DIR="$WORKING_DIR/iree-org"
-BUILD_SCRIPT_DIR="$(dirname -- "$(realpath -- "$0")")/"
+BUILD_SCRIPT_DIR="$(cd -- "$(dirname "${BASH_SOURCE[0]}")" >/dev/null 2>&1 && pwd)"
 UE_SOURCE_DIR="$BUILD_SCRIPT_DIR/../../Source"
 UE_PLUGIN_ROOT="$BUILD_SCRIPT_DIR/../../../.."
 
@@ -45,7 +45,7 @@ echo =========================================
 
 die() { echo "Error: $*" >&2; exit 1; }
 
-commands=(realpath awk rsync p4)
+commands=(awk rsync p4)
 for p in "${commands[@]}"; do
 	command -v "$p" >/dev/null 2>&1 || die "$p is required but not found."
 done
diff --git a/Engine/Plugins/Experimental/NNERuntimeIREE/Source/ExternalBuild/Scripts/Windows/BuildAndInstallWindows.bat b/Engine/Plugins/Experimental/NNERuntimeIREE/Source/ExternalBuild/Scripts/Windows/BuildAndInstallWindows.bat
--- a/Engine/Plugins/Experimental/NNERuntimeIREE/Source/ExternalBuild/Scripts/Windows/BuildAndInstallWindows.bat
+++ b/Engine/Plugins/Experimental/NNERuntimeIREE/Source/ExternalBuild/Scripts/Windows/BuildAndInstallWindows.bat
@@ -14,8 +14,8 @@ CALL "%~dp0_ParseArgs.bat" %* || EXIT /B %ERRORLEVEL%
 :: (Debug)print arguments
 ECHO:
 ECHO Working dir: "%WORKING_DIR%"
-for %%P in (%PLATFORM_LIST:,= %) do (
-	ECHO Building for %%P
+FOR %%i IN (%PLATFORM_LIST% "") DO IF NOT "%%~i"=="" (
+	ECHO Building for %%~i
 )
 ECHO:
 
diff --git a/Engine/Plugins/Experimental/NNERuntimeIREE/Source/ExternalBuild/Scripts/Windows/BuildWindows.bat b/Engine/Plugins/Experimental/NNERuntimeIREE/Source/ExternalBuild/Scripts/Windows/BuildWindows.bat
--- a/Engine/Plugins/Experimental/NNERuntimeIREE/Source/ExternalBuild/Scripts/Windows/BuildWindows.bat
+++ b/Engine/Plugins/Experimental/NNERuntimeIREE/Source/ExternalBuild/Scripts/Windows/BuildWindows.bat
@@ -83,7 +83,7 @@ IF NOT EXIST "%BUILD_DIR%\iree" (
 	
 	FOR %%L IN (%IREE_THIRD_PARTY_LIBRARIES%) DO (
 		IF EXIST "third_party\%%L" (
-			DEL "".\third_party\%%L\CMakeLists.txt"
+			DEL ".\third_party\%%L\CMakeLists.txt"
 			git submodule update --init -- ""./third_party/%%L"
 			IF ERRORLEVEL 1 GOTO :fail
 			IF EXIST "third_party\%%L\third-party" RMDIR /S /Q "third_party\%%L\third-party"
@@ -168,7 +168,7 @@ ECHO =========================================
 ECHO ======== Building NNEMlirTools ==========
 ECHO =========================================
 
-IF NOT EXIST %BUILD_DIR%\NNEMlirTools (
+IF NOT EXIST "%BUILD_DIR%\NNEMlirTools" (
 	SETLOCAL
 	CALL "%VCVARSALL_PATH%" x64 %WIN_SDK_FOR_IREE% -vcvars_ver=%VC_TOOLSET%
 	cmake -G Ninja -B "%BUILD_DIR%/NNEMlirTools" "%UE_SOURCE_DIR%\NNEMlirTools" ^
@@ -270,8 +270,8 @@ IF NOT EXIST "%BUILD_DIR%\iree-runtime-windows" (
 )
 
 :: Run additional platform build scripts
-for %%i in (%PLATFORM_LIST%) do (
-	CALL "%UE_PLUGIN_ROOT%\..\..\..\Platforms\%%i\Plugins\Experimental\NNERuntimeIREE\Source\ExternalBuild\Scripts\Windows\BuildWindows.bat" "%BUILD_DIR%" "%BUILD_TYPE%" "%UE_CMAKE_DIR_SANITIZED%"
+FOR %%i IN (%PLATFORM_LIST% "") DO IF NOT "%%~i"=="" (
+	CALL "%UE_PLUGIN_ROOT%\..\..\..\Platforms\%%~i\Plugins\Experimental\NNERuntimeIREE\Source\ExternalBuild\Scripts\Windows\BuildWindows.bat" "%BUILD_DIR%" "%BUILD_TYPE%" "%UE_CMAKE_DIR_SANITIZED%"
 	IF ERRORLEVEL 1 GOTO :fail
 )
 
diff --git a/Engine/Plugins/Experimental/NNERuntimeIREE/Source/ExternalBuild/Scripts/Windows/Clean.bat b/Engine/Plugins/Experimental/NNERuntimeIREE/Source/ExternalBuild/Scripts/Windows/Clean.bat
deleted file mode 100755
diff --git a/Engine/Plugins/Experimental/NNERuntimeIREE/Source/ExternalBuild/Scripts/Windows/InstallWindows.bat b/Engine/Plugins/Experimental/NNERuntimeIREE/Source/ExternalBuild/Scripts/Windows/InstallWindows.bat
--- a/Engine/Plugins/Experimental/NNERuntimeIREE/Source/ExternalBuild/Scripts/Windows/InstallWindows.bat
+++ b/Engine/Plugins/Experimental/NNERuntimeIREE/Source/ExternalBuild/Scripts/Windows/InstallWindows.bat
@@ -152,18 +152,22 @@ ECHO =========================================
 SETLOCAL ENABLEDELAYEDEXPANSION
 SET "ADDITIONAL_FOLDERS="
 
-for %%i in (%PLATFORM_LIST%) do (
-	CALL :Canonicalize "%UE_PLUGIN_ROOT%\..\..\..\Platforms\%%i\Plugins\Experimental\NNERuntimeIREE" UE_PLATFORM_PLUGIN_ROOT
-	SET "PLATFORM_LIBRARIES_FOLDER=!UE_PLATFORM_PLUGIN_ROOT!\Source\ThirdParty\IREE\Lib\%%i"
+FOR %%i IN (%PLATFORM_LIST% "") DO IF NOT "%%~i"=="" (
+	CALL :Canonicalize "%UE_PLUGIN_ROOT%\..\..\..\Platforms\%%~i\Plugins\Experimental\NNERuntimeIREE" UE_PLATFORM_PLUGIN_ROOT
+	SET "PLATFORM_LIBRARIES_FOLDER=!UE_PLATFORM_PLUGIN_ROOT!\Source\ThirdParty\IREE\Lib\%%~i"
 	ECHO Platform Plugin root: "!UE_PLATFORM_PLUGIN_ROOT!"
 	ECHO Library dir: "!PLATFORM_LIBRARIES_FOLDER!"
 
-	SET "ADDITIONAL_FOLDERS=!ADDITIONAL_FOLDERS! !PLATFORM_LIBRARIES_FOLDER!"
+	IF DEFINED ADDITIONAL_FOLDERS (
+		SET "ADDITIONAL_FOLDERS=!ADDITIONAL_FOLDERS! "!PLATFORM_LIBRARIES_FOLDER!""
+	) ELSE (
+		SET "ADDITIONAL_FOLDERS="!PLATFORM_LIBRARIES_FOLDER!""
+	)
 
-	CALL :CP "%BUILD_DIR%\iree-runtime-%%i\iree\build_tools\third_party\flatcc" "!PLATFORM_LIBRARIES_FOLDER!" *.lib *.a
+	CALL :CP "%BUILD_DIR%\iree-runtime-%%~i\iree\build_tools\third_party\flatcc" "!PLATFORM_LIBRARIES_FOLDER!" *.lib *.a
 	IF ERRORLEVEL 1 GOTO :fail
 
-	CALL :CP "%BUILD_DIR%\iree-runtime-%%i" "!PLATFORM_LIBRARIES_FOLDER!" *.lib *.a
+	CALL :CP "%BUILD_DIR%\iree-runtime-%%~i" "!PLATFORM_LIBRARIES_FOLDER!" *.lib *.a
 	IF ERRORLEVEL 1 GOTO :fail
 
 	SET "PLATFORM_HEADERS_FOLDER_SRC=!UE_PLATFORM_PLUGIN_ROOT!\Source\ExternalBuild\Include"
@@ -244,9 +248,11 @@ IF ERRORLEVEL 1 GOTO :fail
 p4 reconcile -a -e -d -c %CL% "%MLIR_TOOLS_BINARIES_DIR%\..."
 IF ERRORLEVEL 1 GOTO :fail
 
-FOR %%F IN (%ADDITIONAL_FOLDERS%) DO (
-	p4 reconcile -a -e -d -c %CL% "%%F%\..."
-	IF ERRORLEVEL 1 GOTO :fail
+IF DEFINED ADDITIONAL_FOLDERS (
+	FOR %%F IN (!ADDITIONAL_FOLDERS!) DO (
+		p4 reconcile -a -e -d -c %CL% "%%~F%\..."
+		IF ERRORLEVEL 1 GOTO :fail
+	)
 )
 
 :: Note: this is required for the ADDITIONAL_FOLDERS list to be up-to-date
diff --git a/Engine/Plugins/Experimental/NNERuntimeIREE/Source/ExternalBuild/Scripts/Windows/_ParseArgs.bat b/Engine/Plugins/Experimental/NNERuntimeIREE/Source/ExternalBuild/Scripts/Windows/_ParseArgs.bat
--- a/Engine/Plugins/Experimental/NNERuntimeIREE/Source/ExternalBuild/Scripts/Windows/_ParseArgs.bat
+++ b/Engine/Plugins/Experimental/NNERuntimeIREE/Source/ExternalBuild/Scripts/Windows/_ParseArgs.bat
@@ -2,7 +2,7 @@
 SETLOCAL EnableExtensions EnableDelayedExpansion
 
 :: Usage:
-::   _ParseArg.bat [-dir WorkingDir] [-platform "Platform1 Platform2 ..."]
+::   _ParseArgs.bat [-dir WorkingDir] [-platform "Platform1 Platform2 ..."]
 
 :: Defaults
 FOR %%I IN ("%~dp0..\..\Build") DO SET "WORKING_DIR=%%~fI"
@@ -25,7 +25,7 @@ IF /I "%~1"=="-dir" (
 :: -platform LIST (LIST can be "All" or "[Platform1] [Platform2] [Platform3]")
 IF /I "%~1"=="-platform" (
 	SET "ENGINE_PLATFORMS_DIR=%UE5_ROOT%\Engine\Platforms"
-	IF NOT EXIST "%ENGINE_PLATFORMS_DIR%\" EXIT /B 0
+	IF NOT EXIST "%ENGINE_PLATFORMS_DIR%\" EXIT /B 1
 
 	IF "%~2"=="" ( ECHO Error: -platform requires a value.>&2 & EXIT /B 2 )
 	IF /I "%~2" == "All" (
