[NNE] NNERuntimeIREE improve MLIR metadata parser to enumerate all operations that implement FunctionOpInterface and are public.

Problem:
Only MLIR public function operations from the 'func' dialect are enumerated. Input MLIR using function operations from other dialects, e.g. IREE Util are not accepted.

Solution:
Make the implementation more generic and accept all public function operations (all that implement FunctionOpInterface).

#tests On all desktop platforms successfully imported (before) failing MLIR

#rnx

diff --git a/Engine/Plugins/Experimental/NNERuntimeIREE/Source/ExternalBuild/Source/NNEMlirTools/NNERuntimeIREE/Source/NNEMlirTools.cpp b/Engine/Plugins/Experimental/NNERuntimeIREE/Source/ExternalBuild/Source/NNEMlirTools/NNERuntimeIREE/Source/NNEMlirTools.cpp
--- a/Engine/Plugins/Experimental/NNERuntimeIREE/Source/ExternalBuild/Source/NNEMlirTools/NNERuntimeIREE/Source/NNEMlirTools.cpp
+++ b/Engine/Plugins/Experimental/NNERuntimeIREE/Source/ExternalBuild/Source/NNEMlirTools/NNERuntimeIREE/Source/NNEMlirTools.cpp
@@ -22,8 +22,10 @@
 #include "mlir/Dialect/Shape/IR/Shape.h"
 #include "mlir/Dialect/Tensor/IR/Tensor.h"
 #include "mlir/Dialect/Tosa/IR/TosaOps.h"
+#include "mlir/Interfaces/FunctionInterfaces.h"
 #include "mlir/IR/BuiltinOps.h"
 #include "mlir/IR/MLIRContext.h"
+#include "mlir/IR/SymbolTable.h"
 #include "mlir/Parser/Parser.h"
 #include "stablehlo/dialect/ChloOps.h"
 #include "stablehlo/dialect/StablehloOps.h"
@@ -38,7 +40,7 @@
 // * NNEMlirStatus: wrapper around status code & message.
 // * NNEMlirContext: owns shared pointer to mlir::MLIRContext. Multiple modules may share it.
 // * NNEMlirModule: owns shared pointer to mlir::ModuleOp and copy of context shared pointer to keep mlir::MLIRContext alive.
-// * NNEMlirFunction: non-owning view of mlir::func::FuncOp, but also has context and module shared pointers.
+// * NNEMlirFunction: non-owning view of function operation (implements mlir::FunctionOpInterface). But also has context and module shared pointers.
 // * NNEMlirValue: non-owning view of BlockArgument (input) or Type (result). Also has context and module shared pointers.
 //
 // All handles are created with 'new' and released with matching 'NNEMlirReleaseXYZ()'. All handles are safe to destroy in any order
@@ -79,15 +81,17 @@ struct NNEMlirModule_s
 	std::shared_ptr<mlir::MLIRContext> Context;
 	std::shared_ptr<mlir::OwningOpRef<mlir::ModuleOp>> Module;
 
-	std::vector<mlir::func::FuncOp> PublicFuncs;
+	llvm::SmallVector<mlir::FunctionOpInterface, 8> PublicFuncs;
 
 	NNEMlirModule_s(std::shared_ptr<mlir::MLIRContext> Ctx, std::shared_ptr<mlir::OwningOpRef<mlir::ModuleOp>> Mod)
 		: Context(std::move(Ctx)), Module(std::move(Mod))
 	{
-		PublicFuncs.reserve(16);
-		for (mlir::func::FuncOp Func : Module->get().getOps<mlir::func::FuncOp>())
+		mlir::ModuleOp ModOp = Module->get();
+		PublicFuncs.clear();
+
+		for (mlir::FunctionOpInterface Func : ModOp.getOps<mlir::FunctionOpInterface>())
 		{
-			if (Func.isPublic())
+			if (auto Sym = llvm::dyn_cast<mlir::SymbolOpInterface>(Func.getOperation()); Sym && Sym.isPublic())
 			{
 				PublicFuncs.push_back(Func);
 			}
@@ -101,11 +105,11 @@ struct NNEMlirFunction_s
 	std::shared_ptr<mlir::MLIRContext> Context;
 	std::shared_ptr<mlir::OwningOpRef<mlir::ModuleOp>> Module;
 
-	mlir::func::FuncOp Func;
+	mlir::FunctionOpInterface Func;
 
 	std::string CachedName;
 
-	NNEMlirFunction_s(std::shared_ptr<mlir::MLIRContext> Ctx, std::shared_ptr<mlir::OwningOpRef<mlir::ModuleOp>> Mod, mlir::func::FuncOp Fn)
+	NNEMlirFunction_s(std::shared_ptr<mlir::MLIRContext> Ctx, std::shared_ptr<mlir::OwningOpRef<mlir::ModuleOp>> Mod, mlir::FunctionOpInterface Fn)
 		: Context(Ctx), Module(std::move(Mod)), Func(Fn) { }
 };
 
@@ -293,7 +297,7 @@ NNEMlirFunction NNEMlirGetPublicFunction(NNEMlirModule Module, size_t Index)
 	if (!Module) return nullptr;
 	if (Index >= Module->PublicFuncs.size()) return nullptr;
 
-	const mlir::func::FuncOp Func = Module->PublicFuncs[Index];
+	const mlir::FunctionOpInterface Func = Module->PublicFuncs[Index];
 
 	return new NNEMlirFunction_s(Module->Context, Module->Module, Func);
 }
@@ -304,7 +308,7 @@ const char* NNEMlirGetFunctionName(NNEMlirFunction Func)
 
 	if (Func->CachedName.empty())
 	{
-		Func->CachedName = Func->Func.getSymName().str();
+		Func->CachedName = Func->Func.getName().str();
 	}
 
 	return Func->CachedName.c_str();
@@ -339,7 +343,7 @@ NNEMlirValue NNEMlirGetResultValue(NNEMlirFunction Func, size_t Index)
 	if (!Func) return nullptr;
 	if (Index >= Func->Func.getNumResults()) return nullptr;
 
-	const mlir::Type ResultType = Func->Func.getFunctionType().getResult(Index);
+	const mlir::Type ResultType = Func->Func.getResultTypes()[Index];
 
 	return new NNEMlirValue_s(Func->Context, Func->Module, ResultType);
 }
@@ -352,7 +356,7 @@ const char* NNEMlirGetValueName(NNEMlirValue Value)
 	if (Value->CachedName.empty())
 	{
 		const mlir::BlockArgument BlockArg = mlir::cast<mlir::BlockArgument>(Value->Argument);
-		mlir::func::FuncOp Func = mlir::cast<mlir::func::FuncOp>(BlockArg.getOwner()->getParentOp());
+		mlir::FunctionOpInterface Func = mlir::cast<mlir::FunctionOpInterface>(BlockArg.getOwner()->getParentOp());
 
 		const unsigned int Index = BlockArg.getArgNumber();
 		if (auto Attr = Func.getArgAttrOfType<mlir::StringAttr>(Index, "llvm.name"))
diff --git a/Engine/Plugins/Experimental/NNERuntimeIREE/Source/NNERuntimeIREEEditor/Private/NNERuntimeIREEModelDataFactory.cpp b/Engine/Plugins/Experimental/NNERuntimeIREE/Source/NNERuntimeIREEEditor/Private/NNERuntimeIREEModelDataFactory.cpp
--- a/Engine/Plugins/Experimental/NNERuntimeIREE/Source/NNERuntimeIREEEditor/Private/NNERuntimeIREEModelDataFactory.cpp
+++ b/Engine/Plugins/Experimental/NNERuntimeIREE/Source/NNERuntimeIREEEditor/Private/NNERuntimeIREEModelDataFactory.cpp
@@ -159,11 +159,16 @@ UObject* UNNERuntimeIREEModelDataFactory::FactoryCreateBinary(UClass* Class, UOb
 	TConstArrayView<uint8> BufferView = MakeArrayView(Buffer, BufferEnd - Buffer);
 
 	TObjectPtr<UNNERuntimeIREEModuleMetaData> ModuleMetaData = NewObject<UNNERuntimeIREEModuleMetaData>();
-	if (!ModuleMetaData->ParseFromBuffer(BufferView) || ModuleMetaData->FunctionMetaData.IsEmpty())
+	if (!ModuleMetaData->ParseFromBuffer(BufferView))
 	{
 		UE_LOG(LogNNERuntimeIREE, Error, TEXT("UNNERuntimeIREEModelDataFactory failed to parse the models meta data"));
 		return nullptr;
 	}
+	if (ModuleMetaData->FunctionMetaData.IsEmpty())
+	{
+		UE_LOG(LogNNERuntimeIREE, Error, TEXT("UNNERuntimeIREEModelDataFactory no public functions found in models meta data"));
+		return nullptr;
+	}
 
 	if (ModuleMetaData->FunctionMetaData.Num() > 1)
 	{
