[NNE] NNERuntime IREE support of path with spaces on RelTest build on Mac.

#tests Editor, Mac, tested fix with Engine installed at path with spaces and sample project.

#rnx

diff --git a/Engine/Plugins/Experimental/NNERuntimeIREE/Config/IREE_Mac_To_Mac.json b/Engine/Plugins/Experimental/NNERuntimeIREE/Config/IREE_Mac_To_Mac.json
--- a/Engine/Plugins/Experimental/NNERuntimeIREE/Config/IREE_Mac_To_Mac.json
+++ b/Engine/Plugins/Experimental/NNERuntimeIREE/Config/IREE_Mac_To_Mac.json
@@ -7,12 +7,12 @@
 	"Targets": [
 		{
 			"Architecture": "x86_64",
-			"CompilerArguments": "--iree-llvmcpu-target-triple=x86_64-apple-darwin --iree-input-type=auto --iree-hal-target-backends=llvm-cpu --iree-llvmcpu-link-embedded=false --iree-llvmcpu-link-static --output-format=vm-bytecode --iree-llvmcpu-static-library-output-path=${OBJECT_PATH} -o ${VMFB_PATH} ${INPUT_PATH}",
+			"CompilerArguments": "--iree-llvmcpu-target-triple=x86_64-apple-darwin --iree-input-type=auto --iree-hal-target-backends=llvm-cpu --iree-llvmcpu-link-embedded=false --iree-llvmcpu-link-static --iree-llvmcpu-embedded-linker-path=./Binaries/ThirdParty/IREE/Mac/ld.lld --output-format=vm-bytecode --iree-llvmcpu-static-library-output-path=${OBJECT_PATH} -o ${VMFB_PATH} ${INPUT_PATH}",
 			"LinkerArguments": "-dylib -platform_version macos 11.0.0 11.0.0 -arch x86_64 -lSystem -L/Library/Developer/CommandLineTools/SDKs/MacOSX.sdk/usr/lib -o ${SHARED_LIB_PATH} ${OBJECT_PATH}"
 		},
 		{
 			"Architecture": "arm64",
-			"CompilerArguments": "--iree-llvmcpu-target-triple=arm64-apple-darwin --iree-input-type=auto --iree-hal-target-backends=llvm-cpu --iree-llvmcpu-link-embedded=false --iree-llvmcpu-link-static --output-format=vm-bytecode --iree-llvmcpu-static-library-output-path=${OBJECT_PATH} -o ${VMFB_PATH} ${INPUT_PATH}",
+			"CompilerArguments": "--iree-llvmcpu-target-triple=arm64-apple-darwin --iree-input-type=auto --iree-hal-target-backends=llvm-cpu --iree-llvmcpu-link-embedded=false --iree-llvmcpu-link-static --iree-llvmcpu-embedded-linker-path=./Binaries/ThirdParty/IREE/Mac/ld.lld --output-format=vm-bytecode --iree-llvmcpu-static-library-output-path=${OBJECT_PATH} -o ${VMFB_PATH} ${INPUT_PATH}",
 			"LinkerArguments": "-dylib -platform_version macos 11.0.0 11.0.0 -arch arm64 -o ${SHARED_LIB_PATH} ${OBJECT_PATH}"
 		}
 	]
diff --git a/Engine/Plugins/Experimental/NNERuntimeIREE/Source/IREEDriverRDG/Private/IREECompilerRDG.cpp b/Engine/Plugins/Experimental/NNERuntimeIREE/Source/IREEDriverRDG/Private/IREECompilerRDG.cpp
--- a/Engine/Plugins/Experimental/NNERuntimeIREE/Source/IREEDriverRDG/Private/IREECompilerRDG.cpp
+++ b/Engine/Plugins/Experimental/NNERuntimeIREE/Source/IREEDriverRDG/Private/IREECompilerRDG.cpp
@@ -285,6 +285,8 @@ namespace UE::IREE::Compiler
 				FFileHelper::SaveArrayToFile(InFileData, *InputFilePath);
 			}
 
+			const FString PluginDir = FPaths::ConvertRelativePathToFull(*IPluginManager::Get().FindPlugin(UE_PLUGIN_NAME)->GetBaseDir());
+
 			bool bResult = true;
 			for (EShaderPlatform ShaderPlatform : ShaderPlatforms)
 			{
@@ -315,7 +317,7 @@ namespace UE::IREE::Compiler
 				CompilerArguments.ReplaceInline(*FString("${VMFB_PATH}"), *(FString("\"") + VmfbFilePath + "\""));
 				CompilerArguments.ReplaceInline(*FString("${INPUT_PATH}"), *(FString("\"") + InputFilePath + "\""));
 
-				IREEUtils::RunCommand(CompilerCommand, CompilerArguments, IntermediateFilePathNoExt + "_compile-log.txt");
+				IREEUtils::RunCommand(CompilerCommand, CompilerArguments, PluginDir, IntermediateFilePathNoExt + "_compile-log.txt");
 
 				if (!PlatformFile.FileExists(*VmfbFilePath))
 				{
diff --git a/Engine/Plugins/Experimental/NNERuntimeIREE/Source/IREEUtils/Internal/IREEUtils.h b/Engine/Plugins/Experimental/NNERuntimeIREE/Source/IREEUtils/Internal/IREEUtils.h
--- a/Engine/Plugins/Experimental/NNERuntimeIREE/Source/IREEUtils/Internal/IREEUtils.h
+++ b/Engine/Plugins/Experimental/NNERuntimeIREE/Source/IREEUtils/Internal/IREEUtils.h
@@ -9,7 +9,7 @@ namespace UE::IREEUtils
 
 IREEUTILS_API bool ResolveEnvironmentVariables(FString& String);
 
-IREEUTILS_API void RunCommand(const FString& Command, const FString& Arguments, const FString& LogFilePath = FString());
+IREEUTILS_API void RunCommand(const FString& Command, const FString& Arguments, const FString& WorkingDir, const FString& LogFilePath = FString());
 
 IREEUTILS_API bool ImportOnnx(const FString& ImporterCommand, const FString& ImporterArguments, TConstArrayView<uint8> InFileData, const FString& InModelName, const FString& InOutputDir, TArray64<uint8>& OutMlirData);
 
diff --git a/Engine/Plugins/Experimental/NNERuntimeIREE/Source/IREEUtils/Private/IREEUtils.cpp b/Engine/Plugins/Experimental/NNERuntimeIREE/Source/IREEUtils/Private/IREEUtils.cpp
--- a/Engine/Plugins/Experimental/NNERuntimeIREE/Source/IREEUtils/Private/IREEUtils.cpp
+++ b/Engine/Plugins/Experimental/NNERuntimeIREE/Source/IREEUtils/Private/IREEUtils.cpp
@@ -44,12 +44,12 @@ bool ResolveEnvironmentVariables(FString& String)
 	return true;
 }
 
-void RunCommand(const FString& Command, const FString& Arguments, const FString& LogFilePath)
+void RunCommand(const FString& Command, const FString& Arguments, const FString& WorkingDir, const FString& LogFilePath)
 {
 	int32 ReturnCode = 0;
 	bool IsCanceled = false;
 
-	FMonitoredProcess Process(Command, Arguments, true);
+	FMonitoredProcess Process(Command, Arguments, WorkingDir, true);
 	Process.OnCompleted().BindLambda([&ReturnCode] (int32 _ReturnCode) { ReturnCode = _ReturnCode; });
 	Process.OnCanceled().BindLambda([&IsCanceled] (){ IsCanceled = true; });
 
@@ -125,7 +125,7 @@ bool ImportOnnx(const FString& ImporterCommand, const FString& ImporterArguments
 	{
 		SCOPED_NAMED_EVENT_TEXT("Import", FColor::Magenta);
 
-		IREEUtils::RunCommand(ImporterCommand, ImporterArgumentsCopy, IntermediateFilePathNoExt + "_import-log.txt");
+		IREEUtils::RunCommand(ImporterCommand, ImporterArgumentsCopy, FPaths::RootDir(), IntermediateFilePathNoExt + "_import-log.txt");
 	}
 
 	if (!PlatformFile.FileExists(*OutputFilePath))
diff --git a/Engine/Plugins/Experimental/NNERuntimeIREE/Source/NNERuntimeIREE/Private/NNERuntimeIREECompiler.cpp b/Engine/Plugins/Experimental/NNERuntimeIREE/Source/NNERuntimeIREE/Private/NNERuntimeIREECompiler.cpp
--- a/Engine/Plugins/Experimental/NNERuntimeIREE/Source/NNERuntimeIREE/Private/NNERuntimeIREECompiler.cpp
+++ b/Engine/Plugins/Experimental/NNERuntimeIREE/Source/NNERuntimeIREE/Private/NNERuntimeIREECompiler.cpp
@@ -215,6 +215,8 @@ namespace UE::NNERuntimeIREE
 				FFileHelper::SaveArrayToFile(InFileData, *InputFilePath);
 			}
 
+			const FString PluginDir = FPaths::ConvertRelativePathToFull(*IPluginManager::Get().FindPlugin(UE_PLUGIN_NAME)->GetBaseDir());
+
 			bool bResult = true;
 			for (int32 i = 0; i < BuildTargets.Num(); i++)
 			{
@@ -239,7 +241,7 @@ namespace UE::NNERuntimeIREE
 				{
 					SCOPED_NAMED_EVENT_TEXT("Compile", FColor::Magenta);
 
-					IREEUtils::RunCommand(CompilerCommand, CompilerArguments, IntermediateFilePathNoExt + "_compile-log.txt");
+					IREEUtils::RunCommand(CompilerCommand, CompilerArguments, PluginDir, IntermediateFilePathNoExt + "_compile-log.txt");
 				}
 
 				if (!PlatformFile.FileExists(*ObjectFilePath) || !PlatformFile.FileExists(*VmfbFilePath))
@@ -263,7 +265,7 @@ namespace UE::NNERuntimeIREE
 				{
 					SCOPED_NAMED_EVENT_TEXT("Link", FColor::Magenta);
 
-					IREEUtils::RunCommand(LinkerCommand, LinkerArguments, IntermediateFilePathNoExt + "_link-log.txt");
+					IREEUtils::RunCommand(LinkerCommand, LinkerArguments, PluginDir, IntermediateFilePathNoExt + "_link-log.txt");
 				}
 
 				if (!PlatformFile.FileExists(*SharedLibFilePath))
