[NNE] NNERuntimeIREE fix typo in Linux build script.

#tests Ran build script on Linux successfully

#rnx

diff --git a/Engine/Plugins/Experimental/NNERuntimeIREE/Source/ExternalBuild/Scripts/Linux/BuildLinux.sh b/Engine/Plugins/Experimental/NNERuntimeIREE/Source/ExternalBuild/Scripts/Linux/BuildLinux.sh
--- a/Engine/Plugins/Experimental/NNERuntimeIREE/Source/ExternalBuild/Scripts/Linux/BuildLinux.sh
+++ b/Engine/Plugins/Experimental/NNERuntimeIREE/Source/ExternalBuild/Scripts/Linux/BuildLinux.sh
@@ -101,8 +101,8 @@ echo =========================================
 echo ============= Copy Source ===============
 echo =========================================
 
-mkdir -p "$WORKING_DIR/iree/compiler"
-cp -a "$UE_SOURCE_DIR/Iree/Compiler/." "$WORKING_DIR/iree/compiler/" > /dev/null
+mkdir -p "$WORKING_DIR/Iree/Compiler"
+cp -a "$UE_SOURCE_DIR/Iree/Compiler/." "$WORKING_DIR/Iree/Compiler/" > /dev/null
 
 echo
 echo Copied IREE compiler plugin.
