[NNE] NNERuntimeIREERdg always prefer wave32 to be consistent with used GPU profiles from IREE.

#tests ran NNE model tests on all supported platforms

#rnx

diff --git a/Engine/Plugins/Experimental/NNERuntimeIREE/Source/NNERuntimeIREEShader/Internal/NNERuntimeIREEShaderType.h b/Engine/Plugins/Experimental/NNERuntimeIREE/Source/NNERuntimeIREEShader/Internal/NNERuntimeIREEShaderType.h
--- a/Engine/Plugins/Experimental/NNERuntimeIREE/Source/NNERuntimeIREEShader/Internal/NNERuntimeIREEShaderType.h
+++ b/Engine/Plugins/Experimental/NNERuntimeIREE/Source/NNERuntimeIREEShader/Internal/NNERuntimeIREEShaderType.h
@@ -165,6 +165,9 @@ protected:
 		{
 			OutEnvironment.CompilerFlags.Add(CFLAG_AllowRealTypes);
 		}
+
+		// On all supported platforms we prefer wave32
+		OutEnvironment.CompilerFlags.Add(CFLAG_Wave32);
 	}
 #endif // WITH_EDITOR
 };
