[NNE] NNERuntimeIREERDG fixed RDG event name.

#tests Win64, Editor checked name correctly expanded with Unreal Insights.

#rnx

diff --git a/Engine/Plugins/Experimental/NNERuntimeIREE/Source/IREEDriverRDG/Private/IREEDriverRDGDirectCommandBuffer.cpp b/Engine/Plugins/Experimental/NNERuntimeIREE/Source/IREEDriverRDG/Private/IREEDriverRDGDirectCommandBuffer.cpp
--- a/Engine/Plugins/Experimental/NNERuntimeIREE/Source/IREEDriverRDG/Private/IREEDriverRDGDirectCommandBuffer.cpp
+++ b/Engine/Plugins/Experimental/NNERuntimeIREE/Source/IREEDriverRDG/Private/IREEDriverRDGDirectCommandBuffer.cpp
@@ -337,7 +337,7 @@ private:
 		TShaderRef<FNNERuntimeIREEShader> Shader = KernelResource->GetShader(0);
 		if (Shader.IsValid())
 		{
-			RDG_EVENT_SCOPE_STAT(GraphBuilder, FDirectCommandBufferDispatch, "DirectCommandBuffer.Dispatch %s", *KernelName);
+			RDG_EVENT_SCOPE_STAT(GraphBuilder, FDirectCommandBufferDispatch, "DirectCommandBuffer.Dispatch %s", KernelName);
 			RDG_GPU_STAT_SCOPE(GraphBuilder, FDirectCommandBufferDispatch);
 
 			FComputeShaderUtils::AddPass(
