[NNE] NNERuntimeIREERdg add buffer range support.

#tests Win64, Linux, Mac tested affected models runs fine and repeated RDG model tests

#rnx

#FailedValidations Preflight

diff --git a/Engine/Plugins/Experimental/NNERuntimeIREE/Source/IREEDriverRDG/Private/IREEDriverRDGBuffer.cpp b/Engine/Plugins/Experimental/NNERuntimeIREE/Source/IREEDriverRDG/Private/IREEDriverRDGBuffer.cpp
--- a/Engine/Plugins/Experimental/NNERuntimeIREE/Source/IREEDriverRDG/Private/IREEDriverRDGBuffer.cpp
+++ b/Engine/Plugins/Experimental/NNERuntimeIREE/Source/IREEDriverRDG/Private/IREEDriverRDGBuffer.cpp
@@ -77,25 +77,37 @@ public:
 		return Cast(BaseBuffer)->RDGBuffer != FRDGBufferRef{};
 	}
 
-	static const TRefCountPtr<FRDGPooledBuffer>& GetPooledBuffer(iree_hal_buffer_t* BaseBuffer)
+	static FRDGPooledBuffer* GetPooledBuffer(iree_hal_buffer_t* BaseBuffer)
 	{
-		return Cast(BaseBuffer)->PooledBuffer;
+		if (!BaseBuffer)
+		{
+			return nullptr;
+		}
+
+		iree_hal_buffer_t* AllocatedBuffer = iree_hal_buffer_allocated_buffer(BaseBuffer);
+
+		return Cast(AllocatedBuffer)->PooledBuffer.GetReference();
 	}
 
 	static FRDGBufferRef GetRDGBuffer(iree_hal_buffer_t* BaseBuffer, FRDGBuilder* GraphBuilder)
 	{
 		check(GraphBuilder);
 
-		FBuffer* Buffer = Cast(BaseBuffer);
+		if (!BaseBuffer)
+		{
+			return FRDGBufferRef{};
+		}
+
+		iree_hal_buffer_t* AllocatedBuffer = iree_hal_buffer_allocated_buffer(BaseBuffer);
 
+		FBuffer* Buffer = Cast(AllocatedBuffer);
 		if (Buffer->RDGBuffer)
 		{
 			check(Buffer->GraphBuilder == GraphBuilder);
 			return Buffer->RDGBuffer;
 		}
-		{
-			return GraphBuilder->RegisterExternalBuffer(Buffer->PooledBuffer);
-		}
+		
+		return GraphBuilder->RegisterExternalBuffer(Buffer->PooledBuffer);
 	}
 
 private:
@@ -221,7 +233,7 @@ bool BufferIsTransient(iree_hal_buffer_t* Buffer)
 	return Private::FBuffer::IsTransient(Buffer);
 }
 
-const TRefCountPtr<FRDGPooledBuffer>& BufferPooledBufferHandle(iree_hal_buffer_t* Buffer)
+FRDGPooledBuffer* BufferPooledBufferHandle(iree_hal_buffer_t* Buffer)
 {
 	return Private::FBuffer::GetPooledBuffer(Buffer);
 }
diff --git a/Engine/Plugins/Experimental/NNERuntimeIREE/Source/IREEDriverRDG/Private/IREEDriverRDGBuffer.h b/Engine/Plugins/Experimental/NNERuntimeIREE/Source/IREEDriverRDG/Private/IREEDriverRDGBuffer.h
--- a/Engine/Plugins/Experimental/NNERuntimeIREE/Source/IREEDriverRDG/Private/IREEDriverRDGBuffer.h
+++ b/Engine/Plugins/Experimental/NNERuntimeIREE/Source/IREEDriverRDG/Private/IREEDriverRDGBuffer.h
@@ -14,7 +14,7 @@ namespace UE::IREE::HAL::RDG
 
 iree_status_t BufferWrap(iree_allocator_t HostAllocator, iree_hal_allocator_t* DeviceAllocator, iree_hal_memory_type_t MemoryType, iree_hal_memory_access_t AllowedAccess, iree_hal_buffer_usage_t AllowedUsage, iree_device_size_t AllocationSize, iree_device_size_t ByteOffset, iree_device_size_t ByteLength, const TRefCountPtr<FRDGPooledBuffer>& PooledBuffer, iree_hal_buffer_release_callback_t UserReleaseCallback, iree_hal_buffer_t** OutBuffer);
 bool BufferIsTransient(iree_hal_buffer_t* Buffer);
-const TRefCountPtr<FRDGPooledBuffer>& BufferPooledBufferHandle(iree_hal_buffer_t* Buffer);
+FRDGPooledBuffer* BufferPooledBufferHandle(iree_hal_buffer_t* Buffer);
 
 } // namespace UE::IREE::HAL::RDG
 
diff --git a/Engine/Plugins/Experimental/NNERuntimeIREE/Source/IREEDriverRDG/Private/IREEDriverRDGDirectCommandBuffer.cpp b/Engine/Plugins/Experimental/NNERuntimeIREE/Source/IREEDriverRDG/Private/IREEDriverRDGDirectCommandBuffer.cpp
--- a/Engine/Plugins/Experimental/NNERuntimeIREE/Source/IREEDriverRDG/Private/IREEDriverRDGDirectCommandBuffer.cpp
+++ b/Engine/Plugins/Experimental/NNERuntimeIREE/Source/IREEDriverRDG/Private/IREEDriverRDGDirectCommandBuffer.cpp
@@ -24,9 +24,35 @@ namespace UE::IREE::HAL::RDG
 namespace Private
 {
 
-struct FBinding
+struct FSubspanKey
 {
 	FRDGBufferRef Buffer = nullptr;
+
+	uint64 Offset = 0;
+	uint64 Length = 0;
+
+	FSubspanKey(FRDGBufferRef InBuffer, uint64 InOffset, uint64 InLength) :
+		Buffer(InBuffer), Offset(InOffset), Length(InLength)
+	{
+
+	}
+
+	bool operator==(const FSubspanKey& Rhs) const
+	{
+		return Buffer == Rhs.Buffer && Offset == Rhs.Offset && Length == Rhs.Length;
+	}
+};
+
+inline uint32 GetTypeHash(const FSubspanKey& Key)
+{
+	return HashCombine(HashCombine(::GetTypeHash(Key.Buffer), ::GetTypeHash(Key.Offset)), ::GetTypeHash(Key.Length));
+}
+
+struct FTempEntry
+{
+	FRDGBufferRef Buffer = nullptr;
+
+	explicit FTempEntry(FRDGBufferRef InBuffer) : Buffer(InBuffer) {}
 };
 
 class FDirectCommandBuffer
@@ -152,11 +178,12 @@ private:
 		FRDGBuilder& GraphBuilder = DeviceAllocatorGetGraphBuilder(DirectCommandBuffer->DeviceAllocator);
 
 		FRDGBufferRef RDGBuffer = BufferRDGBuffer(TargetRef.buffer, &GraphBuilder);
-		const uint32 BufferSize = RDGBuffer->Desc.GetSize();
-		const uint32 FillOffset = (uint32)TargetRef.offset;
+		const iree_device_size_t FillOffset = iree_hal_buffer_byte_offset(TargetRef.buffer) + TargetRef.offset;
+		const iree_device_size_t BufferLength = iree_hal_buffer_byte_length(TargetRef.buffer);
+		const iree_device_size_t FillLength = TargetRef.length == IREE_HAL_WHOLE_BUFFER ? (BufferLength - TargetRef.offset) : TargetRef.length;
 
-		check(FillOffset < BufferSize);
-		const uint32 FillLength = TargetRef.length == IREE_HAL_WHOLE_BUFFER ? BufferSize - FillOffset : (uint32)TargetRef.length;
+		check(TargetRef.offset <= BufferLength);
+		check(FillLength <= BufferLength - TargetRef.offset);
 
 		uint32 ShaderPattern = 0;
 
@@ -175,7 +202,7 @@ private:
 				return iree_make_status(IREE_STATUS_INVALID_ARGUMENT, "pattern length (%" PRIhsz") is not a power of two or is too large", PatternLength);
 		}
 
-		return BuiltinExecutables::AddFillBufferPass(GraphBuilder, RDGBuffer, ShaderPattern, (uint32)PatternLength, FillOffset, FillLength);
+		return BuiltinExecutables::AddFillBufferPass(GraphBuilder, RDGBuffer, ShaderPattern, (uint32)PatternLength, (uint32)FillOffset, (uint32)FillLength);
 	}
 
 	static iree_status_t UpdateBuffer(iree_hal_command_buffer_t* BaseCommandBuffer, const void* SourceBuffer, iree_host_size_t SourceOffset, iree_hal_buffer_ref_t TargetRef, iree_hal_update_flags_t Flags)
@@ -193,34 +220,45 @@ private:
 #endif
 		SCOPED_NAMED_EVENT_TEXT("FDirectCommandBuffer::CopyBuffer", FColor::Magenta);
 
-		check(SourceRef.length == TargetRef.length);
+		// Host memory not supported
+		check(!iree_all_bits_set(SourceRef.buffer->memory_type, IREE_HAL_MEMORY_TYPE_HOST_LOCAL));
+		check(!iree_all_bits_set(TargetRef.buffer->memory_type, IREE_HAL_MEMORY_TYPE_HOST_LOCAL));
 
 		FDirectCommandBuffer* DirectCommandBuffer = Cast(BaseCommandBuffer);
-
 		FRDGBuilder& GraphBuilder = DeviceAllocatorGetGraphBuilder(DirectCommandBuffer->DeviceAllocator);
 
-		if (iree_all_bits_set(TargetRef.buffer->memory_type, IREE_HAL_MEMORY_TYPE_HOST_LOCAL))
-		{
-			UE_LOG(LogIREEDriverRDG, Warning, TEXT("Skipped buffer readback..."));
-			return iree_ok_status();
-		}
-		
-		checkf(iree_hal_buffer_allocated_buffer(SourceRef.buffer) == SourceRef.buffer, TEXT("Buffer spans not supported yet!"));
-		checkf(iree_hal_buffer_allocated_buffer(TargetRef.buffer) == TargetRef.buffer, TEXT("Buffer spans not supported yet!"));
-
 		FRDGBufferRef SourceRDGBuffer = BufferRDGBuffer(SourceRef.buffer, &GraphBuilder);
 		FRDGBufferRef TargetRDGBuffer = BufferRDGBuffer(TargetRef.buffer, &GraphBuilder);
 
+		const iree_device_size_t SourceOffset = iree_hal_buffer_byte_offset(SourceRef.buffer) + SourceRef.offset;
+		const iree_device_size_t TargetOffset = iree_hal_buffer_byte_offset(TargetRef.buffer) + TargetRef.offset;
+
+		const iree_device_size_t SourceBufferLength = iree_hal_buffer_byte_length(SourceRef.buffer);
+		const iree_device_size_t TargetBufferLength = iree_hal_buffer_byte_length(TargetRef.buffer);
+		
+		const iree_device_size_t ResolvedSourceCopyLength = SourceRef.length == IREE_HAL_WHOLE_BUFFER ? (iree_hal_buffer_byte_length(SourceRef.buffer) - SourceRef.offset) : SourceRef.length;
+		const iree_device_size_t ResolvedTargetCopyLength = TargetRef.length == IREE_HAL_WHOLE_BUFFER ? (iree_hal_buffer_byte_length(TargetRef.buffer) - TargetRef.offset) : TargetRef.length;
+		check(ResolvedSourceCopyLength == ResolvedTargetCopyLength);
+		const iree_device_size_t CopyLength = ResolvedTargetCopyLength;
+
+		check(SourceRef.offset <= SourceBufferLength);
+		check(TargetRef.offset <= TargetBufferLength);
+
+		check(CopyLength <= SourceBufferLength - SourceRef.offset);
+		check(CopyLength <= TargetBufferLength - TargetRef.offset);
+
+		// RHI does not allow copy with Source == Target
 		if (SourceRDGBuffer == TargetRDGBuffer)
 		{
-			FRDGBufferRef TmpBuffer = GraphBuilder.CreateBuffer(SourceRDGBuffer->Desc, TEXT("IREE::CopyBufferTmp"));
+			FRDGBufferDesc BufferCopyDesc = FRDGBufferDesc::CreateByteAddressDesc(CopyLength);
+			FRDGBufferRef TmpBuffer = GraphBuilder.CreateBuffer(BufferCopyDesc, TEXT("IREE::CopyBufferTmp"));
 
-			AddCopyBufferPass(GraphBuilder, TmpBuffer, 0, SourceRDGBuffer, SourceRef.offset, SourceRef.length);
-			AddCopyBufferPass(GraphBuilder, TargetRDGBuffer, TargetRef.offset, TmpBuffer, 0, SourceRef.length);
+			AddCopyBufferPass(GraphBuilder, TmpBuffer, 0, SourceRDGBuffer, SourceOffset, CopyLength);
+			AddCopyBufferPass(GraphBuilder, TargetRDGBuffer, TargetOffset, TmpBuffer, 0, CopyLength);
 		}
 		else
 		{
-			AddCopyBufferPass(GraphBuilder, TargetRDGBuffer, TargetRef.offset, SourceRDGBuffer, SourceRef.offset, SourceRef.length);
+			AddCopyBufferPass(GraphBuilder, TargetRDGBuffer, TargetOffset, SourceRDGBuffer, SourceOffset, CopyLength);
 		}
 
 		return iree_ok_status();
@@ -258,7 +296,7 @@ private:
 		FNNERuntimeIREEShader::FParameters *ShaderParameters = GraphBuilder.AllocParameters<FNNERuntimeIREEShader::FParameters>(ShaderParameterMetadata);
 		uint8* ShaderParameterDataPtr = (uint8*)ShaderParameters;
 
-		TMap<FRDGBufferRef /* Destination */, FRDGBufferRef /* Source */> BufferCopiesToAppend;
+		TMap<FSubspanKey, FTempEntry> BufferCopiesToAppend;
 
 		const uint32 StructSize = ShaderParameterMetadata->GetSize();
 
@@ -296,34 +334,38 @@ private:
 				check(BindingIndex < Bindings.count);
 
 				const iree_hal_buffer_ref_t& BufferRef = Bindings.values[BindingIndex];
-				check(BufferRef.buffer != nullptr);
-				check(BufferRef.offset == 0);
-
-				checkf(iree_hal_buffer_allocated_buffer(BufferRef.buffer) == BufferRef.buffer, TEXT("Buffer spans not supported yet!"));
+				const uint64 BufferOffset = iree_hal_buffer_byte_offset(BufferRef.buffer) + BufferRef.offset;
+				const uint64 BufferRangeLength = iree_hal_buffer_byte_length(BufferRef.buffer);
+				const uint64 BufferLength = BufferRef.length == IREE_HAL_WHOLE_BUFFER ? (BufferRangeLength - BufferRef.offset) : (uint64)BufferRef.length;
 
 				FRDGBufferRef RDGBuffer = BufferRDGBuffer(Bindings.values[BindingIndex].buffer, &GraphBuilder);
 
-				// Require byte address buffer
+				check(BufferLength > 0);
+				check((BufferOffset & 3ull) == 0 && (BufferLength & 3ull) == 0);
+				check(BufferOffset + BufferLength <= RDGBuffer->Desc.GetSize());
+
+				// When we get a buffer range or not a raw buffer we have to create a temp copy to work on and after dispatch copy back
 				FRDGBufferRef RDGWorkingBuffer = nullptr;
-				if (EnumHasAnyFlags(RDGBuffer->Desc.Usage, EBufferUsageFlags::ByteAddressBuffer))
+				if (BufferOffset == 0 && EnumHasAnyFlags(RDGBuffer->Desc.Usage, EBufferUsageFlags::ByteAddressBuffer))
 				{
 					RDGWorkingBuffer = RDGBuffer;
 				}
 				else
 				{
-					auto FindResult = BufferCopiesToAppend.Find(RDGBuffer);
-					if (FindResult)
+					const FSubspanKey Key(RDGBuffer, BufferOffset, BufferLength);
+					if (auto FindResult = BufferCopiesToAppend.Find(Key))
 					{
-						RDGWorkingBuffer = *FindResult;
+						RDGWorkingBuffer = FindResult->Buffer;
 					}
 					else
 					{
-						FRDGBufferDesc BufferCopyDesc = FRDGBufferDesc::CreateByteAddressDesc(RDGBuffer->Desc.GetSize());
+						FRDGBufferDesc BufferCopyDesc = FRDGBufferDesc::CreateByteAddressDesc(Key.Length);
 						FRDGBufferRef BufferCopy = GraphBuilder.CreateBuffer(BufferCopyDesc, TEXT("BufferCopy"));
 
-						AddCopyBufferPass(GraphBuilder, BufferCopy, RDGBuffer);
+						AddCopyBufferPass(GraphBuilder, BufferCopy, 0, RDGBuffer, Key.Offset, Key.Length);
 
-						RDGWorkingBuffer = BufferCopiesToAppend.Emplace(RDGBuffer, BufferCopy);
+						FTempEntry& Entry = BufferCopiesToAppend.Emplace(Key, FTempEntry(BufferCopy));
+						RDGWorkingBuffer = Entry.Buffer;
 					}
 				}
 
@@ -354,9 +396,9 @@ private:
 			UE_LOG(LogIREEDriverRDG, Warning, TEXT("%s: Missing shader for executable %s."), StringCast<TCHAR>(__FUNCTION__).Get(), *KernelName);
 		}
 
-		for (auto BufferCopyResources : BufferCopiesToAppend)
+		for (const auto& Cp : BufferCopiesToAppend)
 		{
-			AddCopyBufferPass(GraphBuilder, BufferCopyResources.Key, BufferCopyResources.Value);
+			AddCopyBufferPass(GraphBuilder, Cp.Key.Buffer, Cp.Key.Offset, Cp.Value.Buffer, 0, Cp.Key.Length);
 		}
 
 		return iree_ok_status();
diff --git a/Engine/Plugins/Experimental/NNERuntimeIREE/Source/NNERuntimeIREE/Private/NNERuntimeIREEModel.cpp b/Engine/Plugins/Experimental/NNERuntimeIREE/Source/NNERuntimeIREE/Private/NNERuntimeIREEModel.cpp
--- a/Engine/Plugins/Experimental/NNERuntimeIREE/Source/NNERuntimeIREE/Private/NNERuntimeIREEModel.cpp
+++ b/Engine/Plugins/Experimental/NNERuntimeIREE/Source/NNERuntimeIREE/Private/NNERuntimeIREEModel.cpp
@@ -1479,30 +1479,39 @@ namespace UE::NNERuntimeIREE
 
 					if (iree_status_is_ok(Status))
 					{
-						for (int i = 0; i < InOutputBindings.Num(); i++)
+						for (int32 i = 0; i < InOutputBindings.Num(); i++)
 						{
-							iree_hal_buffer_view_t* IREEBufferView = NULL;
-
-							if (iree_status_is_ok(Status))
-							{
-								Status = iree_runtime_call_outputs_pop_front_buffer_view(&Call, &IREEBufferView);
-							}
-							
+							iree_hal_buffer_view_t* IREEBufferView = nullptr;
+							Status = iree_runtime_call_outputs_pop_front_buffer_view(&Call, &IREEBufferView);
 							if (!iree_status_is_ok(Status))
 							{
 								break;
 							}
 
+							if (!InOutputBindings[i].Buffer)
+							{
+								iree_hal_buffer_view_release(IREEBufferView);
+								continue;
+							}
+
 							iree_hal_buffer_t* IREEBuffer = iree_hal_buffer_view_buffer(IREEBufferView);
-							iree_device_size_t Offset = iree_hal_buffer_byte_offset(IREEBuffer);
-							iree_device_size_t Length = iree_hal_buffer_byte_length(IREEBuffer);
+							if (!IREEBuffer)
+							{
+								Status = iree_make_status(IREE_STATUS_UNKNOWN, "Failed to get the result buffer");
+								iree_hal_buffer_view_release(IREEBufferView);
+								break;
+							}
 
-							iree_hal_buffer_t* IREEAllocatingBuffer = iree_hal_buffer_allocated_buffer(IREEBuffer); //??
+							const uint64 Length = (uint64)iree_hal_buffer_view_byte_length(IREEBufferView);
+							const uint64 Offset = (uint64)iree_hal_buffer_byte_offset(IREEBuffer);
 
-							FRDGBufferRef RDGBuffer = IREE::HAL::RDG::BufferRDGBuffer(IREEAllocatingBuffer, &GraphBuilder);
+							FRDGBufferRef RDGBuffer = IREE::HAL::RDG::BufferRDGBuffer(IREEBuffer, &GraphBuilder);
+							check(Length <= RDGBuffer->Desc.GetSize() - Offset);
 
-							check(InOutputBindings[i].Buffer);
-							AddCopyBufferPass(GraphBuilder, InOutputBindings[i].Buffer, 0, RDGBuffer, Offset, Length);
+							if (Length <= (uint64)InOutputBindings[i].Buffer->Desc.GetSize())
+							{
+								AddCopyBufferPass(GraphBuilder, InOutputBindings[i].Buffer, 0, RDGBuffer, Offset, Length);
+							}
 
 							iree_hal_buffer_view_release(IREEBufferView);
 						}
