[NNE] NNERuntimeIREECpu fixed typo in size check before copy to output buffer.

#tests Win64, Editor, verified model tests green

#rnx

#FailedValidations Preflight

diff --git a/Engine/Plugins/Experimental/NNERuntimeIREE/Source/NNERuntimeIREE/Private/NNERuntimeIREEModel.cpp b/Engine/Plugins/Experimental/NNERuntimeIREE/Source/NNERuntimeIREE/Private/NNERuntimeIREEModel.cpp
--- a/Engine/Plugins/Experimental/NNERuntimeIREE/Source/NNERuntimeIREE/Private/NNERuntimeIREEModel.cpp
+++ b/Engine/Plugins/Experimental/NNERuntimeIREE/Source/NNERuntimeIREE/Private/NNERuntimeIREEModel.cpp
@@ -990,22 +990,23 @@ namespace UE::NNERuntimeIREE
 							}
 							OutputTensorShapes[i] = UE::NNE::FTensorShape::Make(TConstArrayView<uint32>(Shape, ShapeRank));
 						}
-						if (!InOutputBindings.IsEmpty() && InOutputBindings[i].Data)
+
+						const uint64 DataSizeInBytes = iree_hal_buffer_view_byte_length(BufferView);
+
+						// Copy result to output binding if present and buffer big enough
+						// note: by design silently skip if output buffer is undersized.
+						if (!InOutputBindings.IsEmpty() && InOutputBindings[i].Data && DataSizeInBytes <= InOutputBindings[i].SizeInBytes)
 						{
-							int32 DataSizeInBytes = iree_hal_buffer_view_byte_length(BufferView); 
-							if (InOutputBindings[i].SizeInBytes <= DataSizeInBytes)
+							iree_hal_buffer_t* Buffer = iree_hal_buffer_view_buffer(BufferView);
+							if (!Buffer)
 							{
-								iree_hal_buffer_t* Buffer = iree_hal_buffer_view_buffer(BufferView);
-								if (!Buffer)
-								{
-									iree_hal_buffer_view_release(BufferView);
-									return iree_make_status(IREE_STATUS_UNKNOWN, "Failed to get the result buffer");
-								}
-
-								IREE_RETURN_AND_EVAL_IF_ERROR(iree_hal_buffer_view_release(BufferView),
-									Device->CopyFromBuffer(Buffer, InOutputBindings[i].Data, DataSizeInBytes),
-									"Copy to the output buffer failed");
+								iree_hal_buffer_view_release(BufferView);
+								return iree_make_status(IREE_STATUS_UNKNOWN, "Failed to get the result buffer");
 							}
+
+							IREE_RETURN_AND_EVAL_IF_ERROR(iree_hal_buffer_view_release(BufferView),
+								Device->CopyFromBuffer(Buffer, InOutputBindings[i].Data, DataSizeInBytes),
+								"Copy to the output buffer failed");
 						}
 						iree_hal_buffer_view_release(BufferView);
 					}
@@ -1508,6 +1509,7 @@ namespace UE::NNERuntimeIREE
 							FRDGBufferRef RDGBuffer = IREE::HAL::RDG::BufferRDGBuffer(IREEBuffer, &GraphBuilder);
 							check(Length <= RDGBuffer->Desc.GetSize() - Offset);
 
+							// note: by design silently skip if output buffer is undersized.
 							if (Length <= (uint64)InOutputBindings[i].Buffer->Desc.GetSize())
 							{
 								AddCopyBufferPass(GraphBuilder, InOutputBindings[i].Buffer, 0, RDGBuffer, Offset, Length);
