[NNE] NNERuntimeIREE fix tracing emitting only unique events.

Problem:
Each emitted event was considered "important event" and contributed to the Trace System's cache and were not shown grouped in Unreal Insights. There was also a significant loss in performance, specifically with a lot of events, e.g. with multi-threading.

Solution:
Always use dynamic events and let them cache/group for each unique name (ignoring file and line).

#tests Verified grouping and improved performance

#rnx

diff --git a/Engine/Plugins/Experimental/NNERuntimeIREE/Source/IREETracing/Private/IREETracing.cpp b/Engine/Plugins/Experimental/NNERuntimeIREE/Source/IREETracing/Private/IREETracing.cpp
--- a/Engine/Plugins/Experimental/NNERuntimeIREE/Source/IREETracing/Private/IREETracing.cpp
+++ b/Engine/Plugins/Experimental/NNERuntimeIREE/Source/IREETracing/Private/IREETracing.cpp
@@ -76,7 +76,7 @@ void iree_tracing_set_thread_name(const char* name)
 
 uint32_t iree_tracing_zone_begin(const char* name, const char* file, uint32_t line)
 {
-	uint32_t spec = FCpuProfilerTrace::OutputEventType(name, file, line);
+	uint32_t spec = FCpuProfilerTrace::OutputDynamicEventType(name, file, line);
 	FCpuProfilerTrace::OutputBeginEvent(spec);
 	return spec;
 }
@@ -91,11 +91,6 @@ uint32_t iree_tracing_zone_begin_dynamic(const char* name, size_t name_length, c
 	return spec;
 }
 
-void iree_tracing_zone_end(uint32_t Spec)
-{
-	FCpuProfilerTrace::OutputEndEvent();
-}
-
 uint32_t iree_tracing_zone_begin_external(const char* file_name, size_t file_name_length, uint32_t line, const char* function_name, size_t function_name_length, const char* name, size_t name_length)
 {
 	UE::IREETracing::Private::FAnsiTerminatedSlice function_name_term((const ANSICHAR*)function_name, function_name_length);
@@ -107,6 +102,11 @@ uint32_t iree_tracing_zone_begin_external(const char* file_name, size_t file_nam
 	return iree_tracing_zone_begin(function_name_conv.Get(), file_name_conv.Get(), line);
 }
 
+void iree_tracing_zone_end(uint32_t Spec)
+{
+	FCpuProfilerTrace::OutputEndEvent();
+}
+
 void iree_tracing_publish_source_file(const void* filename, size_t filename_length, const void* content, size_t content_length)
 {
 	// Unimplemented
@@ -123,4 +123,5 @@ uint32_t iree_tracing_zone_begin_dynamic(const char* name, size_t name_length, c
 uint32_t iree_tracing_zone_begin_external(const char* file_name, size_t file_name_length, uint32_t line, const char* function_name, size_t function_name_length, const char* name, size_t name_length) { return 0; }
 void iree_tracing_zone_end(uint32_t Spec) {}
 void iree_tracing_publish_source_file(const void* filename, size_t filename_length, const void* content, size_t content_length) {}
+void iree_tracing_zone_append_text_string_view(uint32_t zone_id, const char* txt, size_t size) {}
 #endif // !UE_BUILD_SHIPPING && CPUPROFILERTRACE_ENABLED
\ No newline at end of file
