[NNE] NNERuntimeIREERdg fixed resource leaking.

Problem:
When exiting Game after loading/running model with IREE RDG an ensure is hit because shader resources are still referenced.

Solution:
When destroying IREE HAL executable, also cleanup shader resources. Check also other HAL objects for resource leakage.

#tests Win64, Game, verified assert not hit anymore because of referenced shader map

#rnx

diff --git a/Engine/Plugins/Experimental/NNERuntimeIREE/Source/IREEDriverRDG/Private/IREEDriverRDGDevice.cpp b/Engine/Plugins/Experimental/NNERuntimeIREE/Source/IREEDriverRDG/Private/IREEDriverRDGDevice.cpp
--- a/Engine/Plugins/Experimental/NNERuntimeIREE/Source/IREEDriverRDG/Private/IREEDriverRDGDevice.cpp
+++ b/Engine/Plugins/Experimental/NNERuntimeIREE/Source/IREEDriverRDG/Private/IREEDriverRDGDevice.cpp
@@ -20,7 +20,6 @@ THIRD_PARTY_INCLUDES_START
 #include "iree/hal/utils/deferred_command_buffer.h"
 #include "iree/hal/utils/file_registry.h"
 #include "iree/hal/utils/file_transfer.h"
-#include "iree/hal/utils/memory_file.h"
 THIRD_PARTY_INCLUDES_END
 #if PLATFORM_MICROSOFT
 #include "Microsoft/HideMicrosoftPlatformAtomics.h"
@@ -93,7 +92,11 @@ private:
 #endif
 		FDevice* Device = Cast(BaseDevice);
 
-		iree_hal_resource_release(Device->DeviceAllocator);
+		Device->Executables.Empty();
+
+		iree_hal_allocator_release(Device->DeviceAllocator);
+
+		iree_hal_channel_provider_release(Device->ChannelProvider);
 		
 		// All arena blocks should have been returned.
 		iree_arena_block_pool_deinitialize(&Device->BlockPool);
diff --git a/Engine/Plugins/Experimental/NNERuntimeIREE/Source/IREEDriverRDG/Private/IREEDriverRDGDeviceAllocator.cpp b/Engine/Plugins/Experimental/NNERuntimeIREE/Source/IREEDriverRDG/Private/IREEDriverRDGDeviceAllocator.cpp
--- a/Engine/Plugins/Experimental/NNERuntimeIREE/Source/IREEDriverRDG/Private/IREEDriverRDGDeviceAllocator.cpp
+++ b/Engine/Plugins/Experimental/NNERuntimeIREE/Source/IREEDriverRDG/Private/IREEDriverRDGDeviceAllocator.cpp
@@ -66,6 +66,8 @@ private:
 		FDeviceAllocator* Allocator = Cast(BaseAllocator);
 
 		FGenericPlatformProcess::ReturnSynchEventToPool(Allocator->Signal);
+		Allocator->Signal = nullptr;
+
 		Allocator->GraphBuilder = nullptr;
 		
 		iree_allocator_free(Allocator->HostAllocator, Allocator);
diff --git a/Engine/Plugins/Experimental/NNERuntimeIREE/Source/IREEDriverRDG/Private/IREEDriverRDGExecutable.cpp b/Engine/Plugins/Experimental/NNERuntimeIREE/Source/IREEDriverRDG/Private/IREEDriverRDGExecutable.cpp
--- a/Engine/Plugins/Experimental/NNERuntimeIREE/Source/IREEDriverRDG/Private/IREEDriverRDGExecutable.cpp
+++ b/Engine/Plugins/Experimental/NNERuntimeIREE/Source/IREEDriverRDG/Private/IREEDriverRDGExecutable.cpp
@@ -1,6 +1,7 @@
 // Copyright Epic Games, Inc. All Rights Reserved.
 
 #include "IREEDriverRDGExecutable.h"
+#include "Templates/UnrealTemplate.h"
 
 #ifdef WITH_IREE_DRIVER_RDG
 
@@ -67,23 +68,21 @@ public:
 		iree_hal_unreal_UnrealShaderDef_vec_t UnrealShaderVec = iree_hal_unreal_ExecutableDef_unreal_shaders_get(ExecutableDef);
 		iree_host_size_t UnrealShaderCount = iree_hal_unreal_UnrealShaderDef_vec_len(UnrealShaderVec);
 
+		TArray<TUniquePtr<FNNERuntimeIREEResource>> KernelResources;
 		for (iree_host_size_t i = 0; i < UnrealShaderCount; i++)
 		{
 			iree_hal_unreal_UnrealShaderDef_table_t UnrealShaderDef = iree_hal_unreal_UnrealShaderDef_vec_at(UnrealShaderVec, i);
 			flatbuffers_string_t SourceFileName = iree_hal_unreal_UnrealShaderDef_source_file_name_get(UnrealShaderDef);
-			// flatbuffers_string_t ModuleName = iree_hal_unreal_UnrealShaderDef_module_name_get(UnrealShaderDef);
-			// flatbuffers_string_t EntryPoint = iree_hal_unreal_UnrealShaderDef_entry_point_get(UnrealShaderDef);
 		
 			const FString ShaderMapFileNameNoExt(ANSI_TO_TCHAR(SourceFileName));
 
 			const TArray<uint8> *Data = Executables->Find(ShaderMapFileNameNoExt);
 			if (!Data)
 			{
+				iree_allocator_free(Executable->HostAllocator, Executable);
 				return iree_make_status(IREE_STATUS_NOT_FOUND, "Executable not found.");
 			}
 
-			Executable->DebugShaderInfos.Add(ShaderMapFileNameNoExt);
-
 			FMemoryReaderView Reader(*Data, /*bIsPersitent =*/ true);
 
 			FIREEDriverRDGShaderParametersMetadata IREEShaderParametersMetadata{};
@@ -108,12 +107,15 @@ public:
 
 			if (!KernelResource->SerializeShaderMap(Reader))
 			{
+				iree_allocator_free(Executable->HostAllocator, Executable);
 				return iree_make_status(IREE_STATUS_NOT_FOUND, "Loaded shader map is not valid/complete.");
 			}
 
-			Executable->KernelResources.Add(MoveTemp(KernelResource));
+			KernelResources.Add(MoveTemp(KernelResource));
 		}
 
+		Executable->KernelResources = MoveTemp(KernelResources);
+
 		*OutExecutable = (iree_hal_executable_t*)Executable;
 		return iree_ok_status();
 	}
@@ -145,6 +147,9 @@ private:
 		UE_LOG(LogIREEDriverRDG, Display, TEXT("%s"), StringCast<TCHAR>(__FUNCTION__).Get());
 #endif
 		FExecutable* Executable = Cast(BaseExecutable);
+
+		Executable->KernelResources.Empty();
+
 		iree_allocator_free(Executable->HostAllocator, Executable);
 	}
 
@@ -162,9 +167,6 @@ private:
 	iree_hal_resource_t Resource;
 	iree_allocator_t HostAllocator;
 	TArray<TUniquePtr<FNNERuntimeIREEResource>> KernelResources;
-
-	// Debug
-	TArray<FString> DebugShaderInfos;
 };
 
 } // namespace Private
