[NNE] NNERuntimeIREE build IREE compiler with the same toolchain as runtime (UE Linux native toolchain v26).

Problem:
On some Horde build machines using IREE compiler fails due to missing libstdc++ that provides GLIBCXX_3.4.32.

Solution:
Build IREE compiler with UE native toolchain for Linux.

#tests validated with nne-build-and-package preflight (all green)

#rnx

diff --git a/Engine/Plugins/Experimental/NNERuntimeIREE/Source/ExternalBuild/Scripts/Linux/BuildLinux.sh b/Engine/Plugins/Experimental/NNERuntimeIREE/Source/ExternalBuild/Scripts/Linux/BuildLinux.sh
--- a/Engine/Plugins/Experimental/NNERuntimeIREE/Source/ExternalBuild/Scripts/Linux/BuildLinux.sh
+++ b/Engine/Plugins/Experimental/NNERuntimeIREE/Source/ExternalBuild/Scripts/Linux/BuildLinux.sh
@@ -22,6 +22,9 @@ UE_PATCHES_DIR=$BUILD_SCRIPT_DIR/../../Patch
 UE_CMAKE_DIR=$BUILD_SCRIPT_DIR/../../CMake
 UE_PLUGIN_ROOT=$BUILD_SCRIPT_DIR/../../../..
 
+UE_PLUGIN_ROOT="$(realpath -e -- "$UE_PLUGIN_ROOT")"
+UE_SOURCE_DIR="$(realpath -e -- "$UE_SOURCE_DIR")"
+
 echo Plugin root dir: "$UE_PLUGIN_ROOT"
 echo Build dir: "$BUILD_DIR"
 echo Source dir: "$UE_SOURCE_DIR"
@@ -47,10 +50,6 @@ C_COMPILER=$BUILD_DIR/Clang/$COMPILER_NAME/$TARGET/bin/clang
 CXX_COMPILER=$BUILD_DIR/Clang/$COMPILER_NAME/$TARGET/bin/clang++
 COMPILER_FLAGS=--target=$TARGET
 
-# We could not build the llvm project with the UE clang so far
-HOST_BINARIES_C_COMPILER=clang-20
-HOST_BINARIES_CXX_COMPILER=clang++-20
-
 # Common rsync options
 # note: -L to copy "resolved" symlinks (e.g. ld.lld)
 RSYNC_OPTIONS=(-a -L)
@@ -146,6 +145,7 @@ if [ ! -d "$BUILD_DIR/iree" ]; then
 	echo Apply git patch to IREE
 	cd "$BUILD_DIR/iree"
 	
+	git apply "$UE_PATCHES_DIR/rocky_linux8.patch"
 	git apply "$UE_PATCHES_DIR/user_provided_tracing_fix.patch"
 	
 	cd ../..
@@ -189,9 +189,9 @@ if [ ! -d "$BUILD_DIR/iree-compiler" ]; then
 		-DIREE_TARGET_BACKEND_DEFAULTS=OFF \
 		-DIREE_TARGET_BACKEND_LLVM_CPU=ON \
 		-DIREE_ERROR_ON_MISSING_SUBMODULES=OFF \
-		-DCMAKE_C_COMPILER=$HOST_BINARIES_C_COMPILER \
+		-DCMAKE_C_COMPILER=$C_COMPILER \
 		-DCMAKE_C_FLAGS=$COMPILER_FLAGS \
-		-DCMAKE_CXX_COMPILER=$HOST_BINARIES_CXX_COMPILER \
+		-DCMAKE_CXX_COMPILER=$CXX_COMPILER \
 		-DCMAKE_CXX_FLAGS=$COMPILER_FLAGS \
 		-DLLVM_ENABLE_TERMINFO=NO \
 		-DIREE_EMBEDDED_RELEASE_INFO=ON \
@@ -232,9 +232,9 @@ if [ ! -d "$BUILD_DIR/NNEMlirTools" ]; then
 		-DIREE_TARGET_BACKEND_DEFAULTS=OFF \
 		-DIREE_TARGET_BACKEND_LLVM_CPU=OFF \
 		-DIREE_ERROR_ON_MISSING_SUBMODULES=OFF \
-		-DCMAKE_C_COMPILER=$HOST_BINARIES_C_COMPILER \
+		-DCMAKE_C_COMPILER=$C_COMPILER \
 		-DCMAKE_C_FLAGS=$COMPILER_FLAGS \
-		-DCMAKE_CXX_COMPILER=$HOST_BINARIES_CXX_COMPILER \
+		-DCMAKE_CXX_COMPILER=$CXX_COMPILER \
 		-DCMAKE_CXX_FLAGS=$COMPILER_FLAGS \
 		-DLLVM_ENABLE_TERMINFO=NO 
 		
