[NNE] NNERuntimeIREECpu on Mac auto-locate SDK and pass to IREE compiler/linker.

#tests Mac, Editor, verified compilation succeds and ran NNEAutomationTest preflight

#rnx

#Synced-CL 48606104

diff --git a/Engine/Plugins/Experimental/NNERuntimeIREE/Config/IREE_Mac_To_Mac.json b/Engine/Plugins/Experimental/NNERuntimeIREE/Config/IREE_Mac_To_Mac.json
--- a/Engine/Plugins/Experimental/NNERuntimeIREE/Config/IREE_Mac_To_Mac.json
+++ b/Engine/Plugins/Experimental/NNERuntimeIREE/Config/IREE_Mac_To_Mac.json
@@ -8,12 +8,12 @@
 		{
 			"Architecture": "x86_64",
 			"CompilerArguments": "--iree-llvmcpu-target-triple=x86_64-apple-darwin --iree-input-type=auto --iree-hal-target-backends=llvm-cpu --iree-llvmcpu-link-embedded=false --iree-llvmcpu-link-static --iree-llvmcpu-embedded-linker-path=./Binaries/ThirdParty/IREE/Mac/ld.lld --output-format=vm-bytecode --iree-llvmcpu-static-library-output-path=${OBJECT_PATH} -o ${VMFB_PATH} ${INPUT_PATH}",
-			"LinkerArguments": "-dylib -platform_version macos 11.0.0 11.0.0 -arch x86_64 -lSystem -L/Library/Developer/CommandLineTools/SDKs/MacOSX.sdk/usr/lib -o ${SHARED_LIB_PATH} ${OBJECT_PATH}"
+			"LinkerArguments": "-dylib -platform_version macos 11.0.0 11.0.0 -arch x86_64 -syslibroot \"${MAC_SDK_ROOT}\" -lSystem -o ${SHARED_LIB_PATH} ${OBJECT_PATH}"
 		},
 		{
 			"Architecture": "arm64",
 			"CompilerArguments": "--iree-llvmcpu-target-triple=arm64-apple-darwin --iree-input-type=auto --iree-hal-target-backends=llvm-cpu --iree-llvmcpu-link-embedded=false --iree-llvmcpu-link-static --iree-llvmcpu-embedded-linker-path=./Binaries/ThirdParty/IREE/Mac/ld.lld --output-format=vm-bytecode --iree-llvmcpu-static-library-output-path=${OBJECT_PATH} -o ${VMFB_PATH} ${INPUT_PATH}",
-			"LinkerArguments": "-dylib -platform_version macos 11.0.0 11.0.0 -arch arm64 -o ${SHARED_LIB_PATH} ${OBJECT_PATH}"
+			"LinkerArguments": "-dylib -platform_version macos 11.0.0 11.0.0 -arch arm64 -syslibroot \"${MAC_SDK_ROOT}\" -lSystem -o ${SHARED_LIB_PATH} ${OBJECT_PATH}"
 		}
 	]
 }
\ No newline at end of file
diff --git a/Engine/Plugins/Experimental/NNERuntimeIREE/Source/IREEUtils/Internal/IREEUtils.h b/Engine/Plugins/Experimental/NNERuntimeIREE/Source/IREEUtils/Internal/IREEUtils.h
--- a/Engine/Plugins/Experimental/NNERuntimeIREE/Source/IREEUtils/Internal/IREEUtils.h
+++ b/Engine/Plugins/Experimental/NNERuntimeIREE/Source/IREEUtils/Internal/IREEUtils.h
@@ -7,6 +7,8 @@
 namespace UE::IREEUtils
 {
 
+IREEUTILS_API bool ResolveSdkPaths(FString& String);
+
 IREEUTILS_API bool ResolveEnvironmentVariables(FString& String);
 
 IREEUTILS_API void RunCommand(const FString& Command, const FString& Arguments, const FString& WorkingDir, const FString& LogFilePath = FString());
diff --git a/Engine/Plugins/Experimental/NNERuntimeIREE/Source/IREEUtils/Private/IREEUtils.cpp b/Engine/Plugins/Experimental/NNERuntimeIREE/Source/IREEUtils/Private/IREEUtils.cpp
--- a/Engine/Plugins/Experimental/NNERuntimeIREE/Source/IREEUtils/Private/IREEUtils.cpp
+++ b/Engine/Plugins/Experimental/NNERuntimeIREE/Source/IREEUtils/Private/IREEUtils.cpp
@@ -2,15 +2,76 @@
 
 #include "IREEUtils.h"
 
+#include "Containers/UnrealString.h"
+#include "HAL/FileManager.h"
 #include "HAL/PlatformFileManager.h"
 #include "IREEUtilsLog.h"
+#include "Logging/LogVerbosity.h"
 #include "Misc/FileHelper.h"
 #include "Misc/MonitoredProcess.h"
+#include "Misc/Optional.h"
 #include "Misc/Paths.h"
 
 namespace UE::IREEUtils
 {
 
+#if PLATFORM_MAC
+namespace Private {
+
+TOptional<FString> GetMacSdkPath()
+{
+	static TOptional<FString> Cached = []() -> TOptional<FString>
+	{
+		// Runs /usr/bin/xcrun --sdk macosx --show-sdk-path and captures stdout
+		FString Output, Error;
+		int32 ReturnCode = -1;
+		const bool bLaunched = FPlatformProcess::ExecProcess(TEXT("/usr/bin/xcrun"), TEXT("--sdk macosx --show-sdk-path"), &ReturnCode, &Output, &Error);
+
+		Output.TrimStartAndEndInline();
+
+		if (bLaunched && ReturnCode == 0 && !Output.IsEmpty() && IFileManager::Get().DirectoryExists(*Output))
+		{
+			UE_LOG(LogIREEUtils, Log, TEXT("Found macOS SDK at %s"), *Output);
+			
+			return Output;
+		}
+
+		if (!Output.IsEmpty())
+		{
+			UE_LOG(LogIREEUtils, Warning, TEXT("xcrun output: %s"), *Output);
+		}
+		UE_LOG(LogIREEUtils, Warning, TEXT("Failed to get macOS SDK path (launched = %d, code = %d). Error: %s"), (bLaunched ? 1 : 0), ReturnCode, *Error);
+
+		return {};
+	}();
+
+	return Cached;
+}
+
+} // namepsace Private		
+#endif // PLATFORM_MAC
+
+bool ResolveSdkPaths(FString& String)
+{
+#if PLATFORM_MAC
+
+	const TOptional<FString> MacSdkPath = Private::GetMacSdkPath();
+
+	if (MacSdkPath.IsSet())
+	{
+		String.ReplaceInline(*FString("${MAC_SDK_ROOT}"), *MacSdkPath.GetValue());
+	}
+	else
+	{
+		if (String.Contains(*FString("${MAC_SDK_ROOT}")))
+		{
+			return false;
+		}
+	}
+#endif // PLATFORM_MAC	
+	return true;
+}
+
 bool ResolveEnvironmentVariables(FString& String)
 {
 	const FString StartString = "$ENV{";
diff --git a/Engine/Plugins/Experimental/NNERuntimeIREE/Source/NNERuntimeIREE/Private/NNERuntimeIREECompiler.cpp b/Engine/Plugins/Experimental/NNERuntimeIREE/Source/NNERuntimeIREE/Private/NNERuntimeIREECompiler.cpp
--- a/Engine/Plugins/Experimental/NNERuntimeIREE/Source/NNERuntimeIREE/Private/NNERuntimeIREECompiler.cpp
+++ b/Engine/Plugins/Experimental/NNERuntimeIREE/Source/NNERuntimeIREE/Private/NNERuntimeIREECompiler.cpp
@@ -234,6 +234,12 @@ namespace UE::NNERuntimeIREE
 					bResult = false;
 					continue;
 				}
+				if (!IREEUtils::ResolveSdkPaths(CompilerArguments))
+				{
+					UE_LOG(LogNNERuntimeIREE, Warning, TEXT("UNNERuntimeIREECpu could not replace SDK paths in %s"), *BuildTargets[i].CompilerArguments);
+					bResult = false;
+					continue;
+				}
 				CompilerArguments.ReplaceInline(*FString("${OBJECT_PATH}"), *(FString("\"") + ObjectFilePath + "\""));
 				CompilerArguments.ReplaceInline(*FString("${VMFB_PATH}"), *(FString("\"") + VmfbFilePath + "\""));
 				CompilerArguments.ReplaceInline(*FString("${INPUT_PATH}"), *(FString("\"") + InputFilePath + "\""));
@@ -259,6 +265,12 @@ namespace UE::NNERuntimeIREE
 					bResult = false;
 					continue;
 				}
+				if (!IREEUtils::ResolveSdkPaths(LinkerArguments))
+				{
+					UE_LOG(LogNNERuntimeIREE, Warning, TEXT("UNNERuntimeIREECpu could not replace SDK paths in %s"), *BuildTargets[i].LinkerArguments);
+					bResult = false;
+					continue;
+				}
 				LinkerArguments.ReplaceInline(*FString("${OBJECT_PATH}"), *(FString("\"") + ObjectFilePath + "\""));
 				LinkerArguments.ReplaceInline(*FString("${SHARED_LIB_PATH}"), *(FString("\"") + SharedLibFilePath + "\""));
 
