[NNE] NNERuntimeIREERdg fix dynamic shapes

#tests dynamic and empty tensor input models
#rnx

A model with dynamic shapes that is tested using NNEModelTests and the IREERDG runtime it throws the following error:
LogNNERuntimeIREE: Error: RDG session: Output binding can be empty or needs to match output tensor descriptor count (got 1, expected 0).
LogNNEModelTests: Error: Invalid number of outputs when running inference on /Game/Models/AutomationTests/OIDN.OIDN using the runtime NNERuntimeIREERdg implementing the interface INNERuntimeRDG

Root cause:
We compare the supplied number of output buffers with the number of OutputTensorsShapes. This is however reset when we have dynamic shapes. Furthermore the shapes are never set after we run the model.

Fixes:
Compare against OutputTensorDescs instead of OutputTensorShapes and set the OutputTensorShapes after the first run of the model. In addition change the interface slightly so that input tensors of size 0 no longer can be null and allow input tensors to have a larger than expected size.

#Synced-CL 48858491

diff --git a/Engine/Plugins/Experimental/NNERuntimeIREE/Source/NNERuntimeIREE/Private/NNERuntimeIREEModel.cpp b/Engine/Plugins/Experimental/NNERuntimeIREE/Source/NNERuntimeIREE/Private/NNERuntimeIREEModel.cpp
--- a/Engine/Plugins/Experimental/NNERuntimeIREE/Source/NNERuntimeIREE/Private/NNERuntimeIREEModel.cpp
+++ b/Engine/Plugins/Experimental/NNERuntimeIREE/Source/NNERuntimeIREE/Private/NNERuntimeIREEModel.cpp
@@ -98,6 +98,34 @@ namespace UE::NNERuntimeIREE
 			}
 		}
 
+		static bool ShapeEqualDims(const UE::NNE::FTensorShape& Shape, const iree_hal_dim_t* Dims, iree_host_size_t Rank)
+		{
+			if (Shape.Rank() != Rank)
+			{
+				return false;
+			}
+			TConstArrayView<uint32> ShapeData = Shape.GetData();
+			for (int i = 0; i < Rank; i++)
+			{
+				if (ShapeData[i] != Dims[i])
+				{
+					return false;
+				}
+			}
+			return true;
+		}
+
+		static UE::NNE::FTensorShape DimsToShape(const iree_hal_dim_t* Dims, iree_host_size_t Rank)
+		{
+			uint32 Shape[UE::NNE::FTensorShape::MaxRank];
+			int32 ShapeRank = FMath::Min((int32)Rank, UE::NNE::FTensorShape::MaxRank);
+			for (int32 ShapeIndex = 0; ShapeIndex < ShapeRank; ShapeIndex++)
+			{
+				Shape[ShapeIndex] = (uint32)Dims[ShapeIndex];
+			}
+			return UE::NNE::FTensorShape::Make(TConstArrayView<uint32>(Shape, ShapeRank));
+		}
+
 		class FInstance
 		{
 		private:
@@ -872,22 +900,6 @@ namespace UE::NNERuntimeIREE
 				}
 
 			private:
-				bool ShapeEqualDims(const UE::NNE::FTensorShape &Shape, const iree_hal_dim_t *Dims, iree_host_size_t Rank)
-				{
-					if (Shape.Rank() != Rank)
-					{
-						return false;
-					}
-					TConstArrayView<uint32> ShapeData = Shape.GetData();
-					for (int i = 0; i < Rank; i++)
-					{
-						if (ShapeData[i] != Dims[i])
-						{
-							return false;
-						}
-					}
-					return true;
-				}
 
 				iree_status_t RunSyncCPUInternal(TConstArrayView<UE::NNE::FTensorBindingCPU> InInputBindings, TConstArrayView<UE::NNE::FTensorBindingCPU> InOutputBindings)
 				{
@@ -980,15 +992,9 @@ namespace UE::NNERuntimeIREE
 							"failed to get output at index %d", i);
 						iree_host_size_t Rank = iree_hal_buffer_view_shape_rank(BufferView);
 						const iree_hal_dim_t* Dims = iree_hal_buffer_view_shape_dims(BufferView);
-						if (!ShapeEqualDims(OutputTensorShapes[i], Dims, Rank))
+						if (!UE::NNERuntimeIREE::Private::ShapeEqualDims(OutputTensorShapes[i], Dims, Rank))
 						{
-							uint32 Shape[UE::NNE::FTensorShape::MaxRank];
-							int32 ShapeRank = FMath::Min((int32)Rank, UE::NNE::FTensorShape::MaxRank);
-							for (int32 ShapeIndex = 0; ShapeIndex < ShapeRank; ShapeIndex++)
-							{
-								Shape[ShapeIndex] = (uint32)Dims[ShapeIndex];
-							}
-							OutputTensorShapes[i] = UE::NNE::FTensorShape::Make(TConstArrayView<uint32>(Shape, ShapeRank));
+							OutputTensorShapes[i] = UE::NNERuntimeIREE::Private::DimsToShape(Dims, Rank);
 						}
 
 						const uint64 DataSizeInBytes = iree_hal_buffer_view_byte_length(BufferView);
@@ -1300,10 +1306,8 @@ namespace UE::NNERuntimeIREE
 				ESetInputTensorShapesStatus SetInputTensorShapes(TConstArrayView<UE::NNE::FTensorShape> InInputShapes)
 				{
 					// OutputTensorShapes will be made available only if all shapes are concretes.
-					InputTensors.Reset();
 					OutputTensors.Reset();
 					OutputTensorShapes.Reset();
-					
 					bool bAllOutputShapeAreConcrete = true;
 					for (int32 i = 0; i < OutputTensorDescs.Num(); i++)
 					{
@@ -1319,10 +1323,11 @@ namespace UE::NNERuntimeIREE
 						}
 					}
 
+					InputTensors.Reset();
 					InputTensorShapes.Reset(InInputShapes.Num());
 					if (InInputShapes.Num() != InputTensorDescs.Num())
 					{
-						UE_LOG(LogNNERuntimeIREE, Error, TEXT("CPU session: Number of input shapes does not match number of input tensors"));
+						UE_LOG(LogNNERuntimeIREE, Error, TEXT("RDG session: Number of input shapes does not match number of input tensors"));
 						return ESetInputTensorShapesStatus::Fail;
 					}
 
@@ -1331,7 +1336,7 @@ namespace UE::NNERuntimeIREE
 						const UE::NNE::FTensorDesc SymbolicDesc = InputTensorDescs[i];
 						if (!InInputShapes[i].IsCompatibleWith(SymbolicDesc.GetShape()))
 						{
-							UE_LOG(LogNNERuntimeIREE, Error, TEXT("CPU session: Input shape does not match input tensor %s of index %d"), *SymbolicDesc.GetName(), i);
+							UE_LOG(LogNNERuntimeIREE, Error, TEXT("RDG session: Input shape does not match input tensor %s of index %d"), *SymbolicDesc.GetName(), i);
 							return ESetInputTensorShapesStatus::Fail;
 						}
 
@@ -1363,35 +1368,24 @@ namespace UE::NNERuntimeIREE
 						return EEnqueueRDGStatus::Fail;
 					}
 
-					if (!InOutputBindings.IsEmpty() && InOutputBindings.Num() != OutputTensorShapes.Num())
+					if (!InOutputBindings.IsEmpty() && InOutputBindings.Num() != OutputTensorDescs.Num())
 					{
-						UE_LOG(LogNNERuntimeIREE, Error, TEXT("RDG session: Output binding can be empty or needs to match output tensor descriptor count (got %d, expected %d)."), InOutputBindings.Num(), OutputTensorShapes.Num());
+						UE_LOG(LogNNERuntimeIREE, Error, TEXT("RDG session: Output binding can be empty or needs to match output tensor descriptor count (got %d, expected %d)."), InOutputBindings.Num(), OutputTensorDescs.Num());
 						return EEnqueueRDGStatus::Fail;
 					}
 
 					for (int32 i = 0; i < InInputBindings.Num(); i++)
 					{
 						const NNE::FTensorBindingRDG& Binding = InInputBindings[i];
-						if (!Binding.Buffer && InputTensors[i].GetDataSize() != 0)
+						if (!Binding.Buffer)
 						{
-							UE_LOG(LogNNERuntimeIREE, Error, TEXT("RDG session: Binding input tensor %d is not set but given size by tensor descriptor is non-zero %d."), i, InputTensors[i].GetDataSize());
+							UE_LOG(LogNNERuntimeIREE, Error, TEXT("RDG session: Binding input tensor %d is null."), i);
 							return EEnqueueRDGStatus::Fail;
 						}
 
-						if (Binding.Buffer && Binding.Buffer->Desc.GetSize() != InputTensors[i].GetDataSize())
+						if (Binding.Buffer->Desc.GetSize() < InputTensors[i].GetDataSize())
 						{
-							UE_LOG(LogNNERuntimeIREE, Error, TEXT("RDG session: Binding input tensor %d size does not match size given by tensor descriptor (got %d, expected %d)."), i, Binding.Buffer->Desc.GetSize(), InputTensorShapes[i].Volume() * InputTensorDescs[i].GetElementByteSize());
-							return EEnqueueRDGStatus::Fail;
-						}
-					}
-
-					for (int32 i = 0; i < InOutputBindings.Num(); i++)
-					{
-						const NNE::FTensorBindingRDG& Binding = InOutputBindings[i];
-
-						if (Binding.Buffer && Binding.Buffer->Desc.GetSize() != OutputTensors[i].GetDataSize())
-						{
-							UE_LOG(LogNNERuntimeIREE, Error, TEXT("Binding output tensor %d size does not match tensor buffer size required (got %d, expected %d)."), i, Binding.Buffer->Desc.GetSize(), OutputTensors[i].GetDataSize());
+							UE_LOG(LogNNERuntimeIREE, Error, TEXT("RDG session: Binding input tensor %d size (%d) is smaller than the size (%d) given by tensor descriptor."), i, Binding.Buffer->Desc.GetSize(), InputTensorShapes[i].Volume() * InputTensorDescs[i].GetElementByteSize());
 							return EEnqueueRDGStatus::Fail;
 						}
 					}
@@ -1480,16 +1474,36 @@ namespace UE::NNERuntimeIREE
 
 					if (iree_status_is_ok(Status))
 					{
-						for (int32 i = 0; i < InOutputBindings.Num(); i++)
+						check(OutputTensorShapes.Num() == 0 || OutputTensorShapes.Num() == OutputTensorDescs.Num());
+						bool bSetOutputTensorShapes = false;
+						if (OutputTensorShapes.Num() == 0 && OutputTensorDescs.Num() > 0)
+						{
+							bSetOutputTensorShapes = true;
+							OutputTensorShapes.SetNum(OutputTensorDescs.Num());
+						}
+						check(OutputTensorShapes.Num() == OutputTensorDescs.Num());
+						check(InOutputBindings.IsEmpty() || InOutputBindings.Num() == OutputTensorDescs.Num());
+
+						for (int32 i = 0; i < OutputTensorDescs.Num(); i++)
 						{
 							iree_hal_buffer_view_t* IREEBufferView = nullptr;
 							Status = iree_runtime_call_outputs_pop_front_buffer_view(&Call, &IREEBufferView);
 							if (!iree_status_is_ok(Status))
 							{
+								if (bSetOutputTensorShapes)
+								{
+									OutputTensorShapes.Reset();
+								}
 								break;
 							}
+							iree_host_size_t Rank = iree_hal_buffer_view_shape_rank(IREEBufferView);
+							const iree_hal_dim_t* Dims = iree_hal_buffer_view_shape_dims(IREEBufferView);
+							if (!UE::NNERuntimeIREE::Private::ShapeEqualDims(OutputTensorShapes[i], Dims, Rank))
+							{
+								OutputTensorShapes[i] = UE::NNERuntimeIREE::Private::DimsToShape(Dims, Rank);
+							}
 
-							if (!InOutputBindings[i].Buffer)
+							if (InOutputBindings.IsEmpty() || !InOutputBindings[i].Buffer)
 							{
 								iree_hal_buffer_view_release(IREEBufferView);
 								continue;
