[NNE] NNERuntimeIREERdg fix bad use of remove during iteration.

#rnx

#Synced-CL 48982312

diff --git a/Engine/Plugins/Experimental/NNERuntimeIREE/Source/NNERuntimeIREEShader/Private/NNERuntimeIREEShaderCompilationManager.cpp b/Engine/Plugins/Experimental/NNERuntimeIREE/Source/NNERuntimeIREEShader/Private/NNERuntimeIREEShaderCompilationManager.cpp
--- a/Engine/Plugins/Experimental/NNERuntimeIREE/Source/NNERuntimeIREEShader/Private/NNERuntimeIREEShaderCompilationManager.cpp
+++ b/Engine/Plugins/Experimental/NNERuntimeIREE/Source/NNERuntimeIREEShader/Private/NNERuntimeIREEShaderCompilationManager.cpp
@@ -361,7 +361,7 @@ void FNNERuntimeIREEShaderCompilationManager::ProcessCompiledNNERuntimeIREEShade
 
 				// Cleanup shader jobs and compile tracking structures
 				ResultArray.Empty();
-				CompiledShaderMaps.Remove(ShaderMap->GetCompilingId());
+				ProcessIt.RemoveCurrent();
 			}
 
 			if (TimeBudget < 0)
