[NNE] NNERuntimeIREERdg preload IREE resources on Game Thread.

Problem:
IREE RDG model initialization is mostly done on the render thread since it requires RDG. Currently this also includes loading of the shader maps. But we should put as little as possible work on the Render Thread and it also triggers a warning in the log.

Solution:
Preload shader maps on the Game Thread.

#test NNE package and test job

#rnx

#Synced-CL 49067668

diff --git a/Engine/Plugins/Experimental/NNERuntimeIREE/Source/IREEDriverRDG/Private/IREEDriverRDGExecutable.cpp b/Engine/Plugins/Experimental/NNERuntimeIREE/Source/IREEDriverRDG/Private/IREEDriverRDGExecutable.cpp
--- a/Engine/Plugins/Experimental/NNERuntimeIREE/Source/IREEDriverRDG/Private/IREEDriverRDGExecutable.cpp
+++ b/Engine/Plugins/Experimental/NNERuntimeIREE/Source/IREEDriverRDG/Private/IREEDriverRDGExecutable.cpp
@@ -1,34 +1,10 @@
 // Copyright Epic Games, Inc. All Rights Reserved.
 
 #include "IREEDriverRDGExecutable.h"
-#include "Templates/UnrealTemplate.h"
 
 #ifdef WITH_IREE_DRIVER_RDG
 
-#include "IREEDriverRDGLog.h"
-#include "IREEDriverRDGShaderParametersMetadata.h"
-#include "Misc/FileHelper.h"
 #include "NNERuntimeIREEShaderShared.h"
-#include "Serialization/MemoryReader.h"
-
-#if PLATFORM_MICROSOFT
-#include "Microsoft/AllowMicrosoftPlatformTypes.h"
-#include "Microsoft/AllowMicrosoftPlatformAtomics.h"
-#endif // PLATFORM_MICROSOFT
-THIRD_PARTY_INCLUDES_START
-#include "iree/hal/utils/executable_debug_info.h"
-
-// flatcc schemas:
-#include "iree/base/internal/flatcc/parsing.h"
-#include "iree/schemas/executable_debug_info_reader.h"
-#include "iree/schemas/executable_debug_info_verifier.h"
-#include "iree/schemas/unreal_executable_def_reader.h"
-#include "iree/schemas/unreal_executable_def_verifier.h"
-THIRD_PARTY_INCLUDES_END
-#if PLATFORM_MICROSOFT
-#include "Microsoft/HideMicrosoftPlatformAtomics.h"
-#include "Microsoft/HideMicrosoftPlatformTypes.h"
-#endif // PLATFORM_MICROSOFT
 
 namespace UE::IREE::HAL::RDG
 {
@@ -36,85 +12,32 @@ namespace UE::IREE::HAL::RDG
 namespace Private
 {
 
-TArray<uint32> GetBufferBindings(const FIREEDriverRDGShaderParametersMetadata& Metadata)
+class FExecutable
 {
-	TArray<uint32> Result;
-	for (const FIREEDriverRDGShaderParametersMetadataEntry& Entry : Metadata.Entries)
+public:
+	iree_hal_resource_t Resource;
+
+private:
+	iree_allocator_t HostAllocator;
+	TArray<TSharedRef<FNNERuntimeIREEResource>> KernelResources;
+
+public:
+	FExecutable(iree_allocator_t HostAllocator, TArrayView<TSharedRef<FNNERuntimeIREEResource>> KernelResources)
+		: HostAllocator(HostAllocator) , KernelResources(KernelResources)
 	{
-		if (Entry.Type == FIREEDriverRDGUniformBufferBaseType::BUFFER_UAV)
-		{
-			Result.Add(Entry.Binding);
-		}
+		iree_hal_resource_initialize((const void*)&FExecutable::VTable, &Resource);
 	}
-	return Result;
-}
 
-class FExecutable
-{
-public:
-	static iree_status_t Create(iree_allocator_t HostAllocator, const TMap<FString, TArray<uint8>>* Executables, const iree_hal_executable_params_t* ExecutableParams, iree_hal_executable_t** OutExecutable)
+	static iree_status_t Create(iree_allocator_t HostAllocator, TArrayView<TSharedRef<FNNERuntimeIREEResource>> KernelResources, iree_hal_executable_t** OutExecutable)
 	{
 		SCOPED_NAMED_EVENT_TEXT("FExecutable::Create", FColor::Magenta);
 
-		check(Executables);
 		check(OutExecutable);
 
 		FExecutable* Executable;
 		IREE_RETURN_IF_ERROR(iree_allocator_malloc(HostAllocator, sizeof(*Executable), (void**)&Executable));
-		iree_hal_resource_initialize((const void*)&FExecutable::VTable, &Executable->Resource);
-		Executable->HostAllocator = HostAllocator;
-
-		iree_hal_unreal_ExecutableDef_table_t ExecutableDef = iree_hal_unreal_ExecutableDef_as_root(ExecutableParams->executable_data.data);
-		iree_hal_unreal_UnrealShaderDef_vec_t UnrealShaderVec = iree_hal_unreal_ExecutableDef_unreal_shaders_get(ExecutableDef);
-		iree_host_size_t UnrealShaderCount = iree_hal_unreal_UnrealShaderDef_vec_len(UnrealShaderVec);
-
-		TArray<TUniquePtr<FNNERuntimeIREEResource>> KernelResources;
-		for (iree_host_size_t i = 0; i < UnrealShaderCount; i++)
-		{
-			iree_hal_unreal_UnrealShaderDef_table_t UnrealShaderDef = iree_hal_unreal_UnrealShaderDef_vec_at(UnrealShaderVec, i);
-			flatbuffers_string_t SourceFileName = iree_hal_unreal_UnrealShaderDef_source_file_name_get(UnrealShaderDef);
 		
-			const FString ShaderMapFileNameNoExt(ANSI_TO_TCHAR(SourceFileName));
-
-			const TArray<uint8> *Data = Executables->Find(ShaderMapFileNameNoExt);
-			if (!Data)
-			{
-				iree_allocator_free(Executable->HostAllocator, Executable);
-				return iree_make_status(IREE_STATUS_NOT_FOUND, "Executable not found.");
-			}
-
-			FMemoryReaderView Reader(*Data, /*bIsPersitent =*/ true);
-
-			FIREEDriverRDGShaderParametersMetadata IREEShaderParametersMetadata{};
-			FIREEDriverRDGShaderParametersMetadata::StaticStruct()->SerializeBin(Reader, &IREEShaderParametersMetadata);
-
-			TUniquePtr<FNNERuntimeIREEShaderParametersMetadataAllocations> ShaderParameterMetadataAllocations = MakeUnique<FNNERuntimeIREEShaderParametersMetadataAllocations>();
-			FShaderParametersMetadata* ShaderParametersMetadata = HAL::RDG::BuildShaderParametersMetadata(IREEShaderParametersMetadata, *ShaderParameterMetadataAllocations);
-
-			TUniquePtr<FNNERuntimeIREEResource> KernelResource = MakeUnique<FNNERuntimeIREEResource>();
-			KernelResource->SetupResource(
-				GMaxRHIFeatureLevel,
-				ShaderMapFileNameNoExt,
-				FString(),
-				FString(),
-				FString(),
-				MoveTemp(ShaderParameterMetadataAllocations),
-				ShaderParametersMetadata,
-				FName(),
-				GetBufferBindings(IREEShaderParametersMetadata)
-			);
-			
-
-			if (!KernelResource->SerializeShaderMap(Reader))
-			{
-				iree_allocator_free(Executable->HostAllocator, Executable);
-				return iree_make_status(IREE_STATUS_NOT_FOUND, "Loaded shader map is not valid/complete.");
-			}
-
-			KernelResources.Add(MoveTemp(KernelResource));
-		}
-
-		Executable->KernelResources = MoveTemp(KernelResources);
+		new (Executable) FExecutable(HostAllocator, KernelResources);
 
 		*OutExecutable = (iree_hal_executable_t*)Executable;
 		return iree_ok_status();
@@ -129,7 +52,7 @@ public:
 			return iree_make_status(IREE_STATUS_OUT_OF_RANGE, "Invalid entry point index %i requested from executable", EntryPoint);
 		}
 
-		*OutResource = Executable->KernelResources[EntryPoint].Get();
+		*OutResource = &(Executable->KernelResources[EntryPoint].Get());
 
 		return iree_ok_status();
 	}
@@ -148,9 +71,11 @@ private:
 #endif
 		FExecutable* Executable = Cast(BaseExecutable);
 
-		Executable->KernelResources.Empty();
+		iree_allocator_t HostAllocator = Executable->HostAllocator;
+
+		Executable->~FExecutable();
 
-		iree_allocator_free(Executable->HostAllocator, Executable);
+		iree_allocator_free(HostAllocator, Executable);
 	}
 
 	inline static const iree_hal_executable_vtable_t VTable = []
@@ -162,18 +87,13 @@ private:
 		return Result;
 
 	}();
-
-private:
-	iree_hal_resource_t Resource;
-	iree_allocator_t HostAllocator;
-	TArray<TUniquePtr<FNNERuntimeIREEResource>> KernelResources;
 };
 
 } // namespace Private
 
-iree_status_t ExecutableCreate(iree_allocator_t HostAllocator, const TMap<FString, TArray<uint8>>* Executables, const iree_hal_executable_params_t* ExecutableParams, iree_hal_executable_t** OutExecutable)
+iree_status_t ExecutableCreate(iree_allocator_t HostAllocator, TArrayView<TSharedRef<FNNERuntimeIREEResource>> KernelResources, iree_hal_executable_t** OutExecutable)
 {
-	return Private::FExecutable::Create(HostAllocator, Executables, ExecutableParams, OutExecutable);
+	return Private::FExecutable::Create(HostAllocator, KernelResources, OutExecutable);
 }
 
 iree_status_t ExecutableGetResource(iree_hal_executable_t* Executable, int32 EntryPoint, const FNNERuntimeIREEResource** OutResource)
@@ -181,6 +101,6 @@ iree_status_t ExecutableGetResource(iree_hal_executable_t* Executable, int32 Ent
 	return Private::FExecutable::GetResource(Executable, EntryPoint, OutResource);
 }
 
-} // UE::IREE
+} // UE::IREE::HAL::RDG
 
 #endif // WITH_IREE_DRIVER_RDG
\ No newline at end of file
diff --git a/Engine/Plugins/Experimental/NNERuntimeIREE/Source/IREEDriverRDG/Private/IREEDriverRDGExecutable.h b/Engine/Plugins/Experimental/NNERuntimeIREE/Source/IREEDriverRDG/Private/IREEDriverRDGExecutable.h
--- a/Engine/Plugins/Experimental/NNERuntimeIREE/Source/IREEDriverRDG/Private/IREEDriverRDGExecutable.h
+++ b/Engine/Plugins/Experimental/NNERuntimeIREE/Source/IREEDriverRDG/Private/IREEDriverRDGExecutable.h
@@ -5,13 +5,14 @@
 #ifdef WITH_IREE_DRIVER_RDG
 
 #include "IREEDriverRDG.h"
+#include "Templates/SharedPointer.h"
 
 class FNNERuntimeIREEResource;
 
 namespace UE::IREE::HAL::RDG
 {
 
-iree_status_t ExecutableCreate(iree_allocator_t HostAllocator, const TMap<FString, TArray<uint8>>* Executables, const iree_hal_executable_params_t* ExecutableParams, iree_hal_executable_t** OutExecutable);
+iree_status_t ExecutableCreate(iree_allocator_t HostAllocator, TArrayView<TSharedRef<FNNERuntimeIREEResource>> KernelResources, iree_hal_executable_t** OutExecutable);
 
 iree_status_t ExecutableGetResource(iree_hal_executable_t* Executable, int32 EntryPoint, const FNNERuntimeIREEResource** OutResource);
 
diff --git a/Engine/Plugins/Experimental/NNERuntimeIREE/Source/IREEDriverRDG/Private/IREEDriverRDGNoOpExecutableCache.cpp b/Engine/Plugins/Experimental/NNERuntimeIREE/Source/IREEDriverRDG/Private/IREEDriverRDGNoOpExecutableCache.cpp
--- a/Engine/Plugins/Experimental/NNERuntimeIREE/Source/IREEDriverRDG/Private/IREEDriverRDGNoOpExecutableCache.cpp
+++ b/Engine/Plugins/Experimental/NNERuntimeIREE/Source/IREEDriverRDG/Private/IREEDriverRDGNoOpExecutableCache.cpp
@@ -6,6 +6,29 @@
 
 #include "IREEDriverRDGExecutable.h"
 #include "IREEDriverRDGLog.h"
+#include "Serialization/MemoryReader.h"
+#include "Templates/SharedPointer.h"
+#include "IREEDriverRDGShaderParametersMetadata.h"
+#include "NNERuntimeIREEShaderShared.h"
+
+#if PLATFORM_MICROSOFT
+#include "Microsoft/AllowMicrosoftPlatformTypes.h"
+#include "Microsoft/AllowMicrosoftPlatformAtomics.h"
+#endif // PLATFORM_MICROSOFT
+THIRD_PARTY_INCLUDES_START
+#include "iree/hal/utils/executable_debug_info.h"
+
+// flatcc schemas:
+#include "iree/base/internal/flatcc/parsing.h"
+#include "iree/schemas/executable_debug_info_reader.h"
+#include "iree/schemas/executable_debug_info_verifier.h"
+#include "iree/schemas/unreal_executable_def_reader.h"
+#include "iree/schemas/unreal_executable_def_verifier.h"
+THIRD_PARTY_INCLUDES_END
+#if PLATFORM_MICROSOFT
+#include "Microsoft/HideMicrosoftPlatformAtomics.h"
+#include "Microsoft/HideMicrosoftPlatformTypes.h"
+#endif // PLATFORM_MICROSOFT
 
 namespace UE::IREE::HAL::RDG
 {
@@ -13,11 +36,71 @@ namespace UE::IREE::HAL::RDG
 namespace Private
 {
 
+TArray<uint32> GetBufferBindings(const FIREEDriverRDGShaderParametersMetadata& Metadata)
+{
+	TArray<uint32> Result;
+	for (const FIREEDriverRDGShaderParametersMetadataEntry& Entry : Metadata.Entries)
+	{
+		if (Entry.Type == FIREEDriverRDGUniformBufferBaseType::BUFFER_UAV)
+		{
+			Result.Add(Entry.Binding);
+		}
+	}
+	return Result;
+}
+
+iree_status_t PreloadIREEResources(const TMap<FString, TArray<uint8>>* Executables, TMap<FString, TSharedRef<FNNERuntimeIREEResource>>& OutPreloadedResources)
+{
+	SCOPED_NAMED_EVENT_TEXT("PreloadIREEResources", FColor::Magenta);
+
+	for (const TPair<FString, TArray<uint8>>& Pair : *Executables)
+	{
+		const FString& Name = Pair.Key;
+		const TArray<uint8>& Data = Pair.Value;
+
+		TArray<TSharedRef<FNNERuntimeIREEResource>> Resources;
+		
+		FMemoryReaderView Reader(Data, /*bIsPersitent =*/ true);
+
+		FIREEDriverRDGShaderParametersMetadata IREEShaderParametersMetadata{};
+		FIREEDriverRDGShaderParametersMetadata::StaticStruct()->SerializeBin(Reader, &IREEShaderParametersMetadata);
+
+		TUniquePtr<FNNERuntimeIREEShaderParametersMetadataAllocations> ShaderParameterMetadataAllocations = MakeUnique<FNNERuntimeIREEShaderParametersMetadataAllocations>();
+		FShaderParametersMetadata* ShaderParametersMetadata = HAL::RDG::BuildShaderParametersMetadata(IREEShaderParametersMetadata, *ShaderParameterMetadataAllocations);
+
+		TSharedRef<FNNERuntimeIREEResource> KernelResource = MakeShared<FNNERuntimeIREEResource>();
+		KernelResource->SetupResource(
+			GMaxRHIFeatureLevel,
+			Name,
+			FString(),
+			FString(),
+			FString(),
+			MoveTemp(ShaderParameterMetadataAllocations),
+			ShaderParametersMetadata,
+			FName(),
+			GetBufferBindings(IREEShaderParametersMetadata)
+		);
+		
+		if (!KernelResource->SerializeShaderMap(Reader))
+		{
+			return iree_make_status(IREE_STATUS_NOT_FOUND, "Loaded shader map is not valid/complete.");
+		}
+
+		OutPreloadedResources.Add(Name, KernelResource);
+	}
+
+	return iree_ok_status();
+}
+
 class FNoOpExecutableCache
 {
 public:
 	static iree_status_t Create(iree_allocator_t HostAllocator, const TMap<FString, TArray<uint8>>* Executables, iree_hal_executable_cache_t** OutExecutableCache)
 	{
+		SCOPED_NAMED_EVENT_TEXT("FNoOpExecutableCache::Create", FColor::Magenta);
+
+		checkf(IsInGameThread(), TEXT("IREE Executable Cache must be created on Game Thread"));
+
 		check(Executables);
 		check(OutExecutableCache);
 
@@ -25,7 +108,13 @@ public:
 		IREE_RETURN_IF_ERROR(iree_allocator_malloc(HostAllocator, sizeof(*ExecutableCache), (void**)&ExecutableCache));
 		iree_hal_resource_initialize((const void*)&FNoOpExecutableCache::VTable, &ExecutableCache->Resource);
 		ExecutableCache->HostAllocator = HostAllocator;
-		ExecutableCache->Executables = Executables;
+		
+		iree_status_t Status = PreloadIREEResources(Executables, ExecutableCache->PreloadedResources);
+		if (!iree_status_is_ok(Status))
+		{
+			iree_allocator_free(ExecutableCache->HostAllocator, ExecutableCache);
+			return Status;
+		}
 
 		*OutExecutableCache = (iree_hal_executable_cache_t*)ExecutableCache;
 		return iree_ok_status();
@@ -39,11 +128,14 @@ private:
 	}
 
 	static void Destroy(iree_hal_executable_cache_t* BaseExecutableCache)
-	{
+	{		
 #if IREE_DRIVER_RDG_VERBOSITY == 1
 		UE_LOG(LogIREEDriverRDG, Display, TEXT("%s"), StringCast<TCHAR>(__FUNCTION__).Get());
 #endif
 		FNoOpExecutableCache* ExecutableCache = Cast(BaseExecutableCache);
+
+		ExecutableCache->PreloadedResources.Empty();
+
 		iree_allocator_free(ExecutableCache->HostAllocator, ExecutableCache);
 	}
 
@@ -75,14 +167,33 @@ private:
 
 		FNoOpExecutableCache* ExecutableCache = Cast(BaseExecutableCache);
 
-		return ExecutableCreate(ExecutableCache->HostAllocator, ExecutableCache->Executables, ExecutableParams, OutExecutable);
+		iree_hal_unreal_ExecutableDef_table_t ExecutableDef = iree_hal_unreal_ExecutableDef_as_root(ExecutableParams->executable_data.data);
+		iree_hal_unreal_UnrealShaderDef_vec_t UnrealShaderVec = iree_hal_unreal_ExecutableDef_unreal_shaders_get(ExecutableDef);
+		iree_host_size_t UnrealShaderCount = iree_hal_unreal_UnrealShaderDef_vec_len(UnrealShaderVec);
+
+		TArray<TSharedRef<FNNERuntimeIREEResource>> KernelResources;
+		for (iree_host_size_t i = 0; i < UnrealShaderCount; i++)
+		{
+			iree_hal_unreal_UnrealShaderDef_table_t UnrealShaderDef = iree_hal_unreal_UnrealShaderDef_vec_at(UnrealShaderVec, i);
+			flatbuffers_string_t SourceFileName = iree_hal_unreal_UnrealShaderDef_source_file_name_get(UnrealShaderDef);
+
+			FString ShaderName = FString(UTF8_TO_TCHAR(SourceFileName));
+			TSharedRef<FNNERuntimeIREEResource>* FoundResources = ExecutableCache->PreloadedResources.Find(ShaderName);
+			if (!FoundResources)
+			{
+				return iree_make_status(IREE_STATUS_NOT_FOUND, "No preloaded IREE resource found for shader '%s'", SourceFileName);
+			}
+			KernelResources.Add(*FoundResources);
+		}
+
+		return ExecutableCreate(ExecutableCache->HostAllocator, KernelResources, OutExecutable);
 	}
 
 	static const iree_hal_executable_cache_vtable_t VTable;
 
 	iree_hal_resource_t Resource;
 	iree_allocator_t HostAllocator;
-	const TMap<FString, TArray<uint8>>* Executables;
+	TMap<FString, TSharedRef<FNNERuntimeIREEResource>> PreloadedResources;
 };
 
 const iree_hal_executable_cache_vtable_t FNoOpExecutableCache::VTable = 
