[NNE] IREERCPU add multiple targets based on x86 features

#rnx
#tests NNEModelTests

Add support to recognize x86 features during runtime. This allows for more optimization options without sacrificing compatibility with older devices.

#Synced-CL 49325839

diff --git a/Engine/Plugins/Experimental/NNERuntimeIREE/Config/IREE_Linux_To_Linux.json b/Engine/Plugins/Experimental/NNERuntimeIREE/Config/IREE_Linux_To_Linux.json
--- a/Engine/Plugins/Experimental/NNERuntimeIREE/Config/IREE_Linux_To_Linux.json
+++ b/Engine/Plugins/Experimental/NNERuntimeIREE/Config/IREE_Linux_To_Linux.json
@@ -6,9 +6,31 @@
 	"SharedLibExt": ".so",
 	"Targets": [
 		{
+			"Name": "AVX512",
 			"Architecture": "",
-			"CompilerArguments": "--iree-llvmcpu-target-triple=x86_64-unknown-linux-gnu --iree-input-type=auto --iree-hal-target-backends=llvm-cpu --iree-llvmcpu-link-embedded=false --iree-llvmcpu-link-static --output-format=vm-bytecode --iree-llvmcpu-static-library-output-path=${OBJECT_PATH} -o ${VMFB_PATH} ${INPUT_PATH}",
-			"LinkerArguments": "-shared --export-dynamic -o ${SHARED_LIB_PATH} ${OBJECT_PATH}"
+			"x86Features": [ "avx512", "fma", "f16c" ],
+			"CompilerArguments": "--iree-llvmcpu-target-triple=x86_64-unknown-linux-gnu --iree-input-type=auto --iree-hal-target-backends=llvm-cpu --iree-llvmcpu-target-cpu-features=+sse,+sse2,+ssse3,+sse4.1,+sse4.2,+avx,+f16c,+avx2,+fma,+avx512f,+avx512vl,+avx512bw,+avx512dq --iree-llvmcpu-link-embedded=false --iree-llvmcpu-link-static --output-format=vm-bytecode --iree-llvmcpu-static-library-output-path=${OBJECT_PATH} -o ${VMFB_PATH} ${INPUT_PATH}",
+			"LinkerArguments": "-shared -o ${SHARED_LIB_PATH} ${OBJECT_PATH}"
+		},
+		{
+			"Name": "AVX2",
+			"Architecture": "",
+			"x86Features": [ "avx2", "fma", "f16c" ],
+			"CompilerArguments": "--iree-llvmcpu-target-triple=x86_64-unknown-linux-gnu --iree-input-type=auto --iree-hal-target-backends=llvm-cpu --iree-llvmcpu-target-cpu-features=+sse,+sse2,+ssse3,+sse4.1,+sse4.2,+avx,+f16c,+avx2,+fma --iree-llvmcpu-link-embedded=false --iree-llvmcpu-link-static --output-format=vm-bytecode --iree-llvmcpu-static-library-output-path=${OBJECT_PATH} -o ${VMFB_PATH} ${INPUT_PATH}",
+			"LinkerArguments": "-shared -o ${SHARED_LIB_PATH} ${OBJECT_PATH}"
+		},
+		{
+			"Name": "AVX",
+			"Architecture": "",
+			"x86Features": [ "avx", "f16c" ],
+			"CompilerArguments": "--iree-llvmcpu-target-triple=x86_64-unknown-linux-gnu --iree-input-type=auto --iree-hal-target-backends=llvm-cpu --iree-llvmcpu-target-cpu-features=+sse,+sse2,+ssse3,+sse4.1,+sse4.2,+avx,+f16c --iree-llvmcpu-link-embedded=false --iree-llvmcpu-link-static --output-format=vm-bytecode --iree-llvmcpu-static-library-output-path=${OBJECT_PATH} -o ${VMFB_PATH} ${INPUT_PATH}",
+			"LinkerArguments": "-shared -o ${SHARED_LIB_PATH} ${OBJECT_PATH}"
+		},
+		{
+			"Name": "Default",
+			"Architecture": "",
+			"CompilerArguments": "--iree-llvmcpu-target-triple=x86_64-unknown-linux-gnu --iree-input-type=auto --iree-hal-target-backends=llvm-cpu --iree-llvmcpu-target-cpu-features=+sse,+sse2,+ssse3,+sse4.1,+sse4.2 --iree-llvmcpu-link-embedded=false --iree-llvmcpu-link-static --output-format=vm-bytecode --iree-llvmcpu-static-library-output-path=${OBJECT_PATH} -o ${VMFB_PATH} ${INPUT_PATH}",
+			"LinkerArguments": "-shared -o ${SHARED_LIB_PATH} ${OBJECT_PATH}"
 		}
 	]
 }
diff --git a/Engine/Plugins/Experimental/NNERuntimeIREE/Config/IREE_Mac_To_Mac.json b/Engine/Plugins/Experimental/NNERuntimeIREE/Config/IREE_Mac_To_Mac.json
--- a/Engine/Plugins/Experimental/NNERuntimeIREE/Config/IREE_Mac_To_Mac.json
+++ b/Engine/Plugins/Experimental/NNERuntimeIREE/Config/IREE_Mac_To_Mac.json
@@ -6,14 +6,23 @@
 	"SharedLibExt": ".dylib",
 	"Targets": [
 		{
+			"Name":  "ARM",
+			"Architecture": "arm64",
+			"CompilerArguments": "--iree-llvmcpu-target-triple=arm64-apple-darwin --iree-input-type=auto --iree-hal-target-backends=llvm-cpu --iree-llvmcpu-link-embedded=false --iree-llvmcpu-link-static --iree-llvmcpu-embedded-linker-path=./Binaries/ThirdParty/IREE/Mac/ld.lld --output-format=vm-bytecode --iree-llvmcpu-static-library-output-path=${OBJECT_PATH} -o ${VMFB_PATH} ${INPUT_PATH}",
+			"LinkerArguments": "-dylib -platform_version macos 11.0.0 11.0.0 -arch arm64 -syslibroot \"${MAC_SDK_ROOT}\" -lSystem -o ${SHARED_LIB_PATH} ${OBJECT_PATH}"
+		},
+		{
+			"Name": "Intel_AVX2",
 			"Architecture": "x86_64",
-			"CompilerArguments": "--iree-llvmcpu-target-triple=x86_64-apple-darwin --iree-input-type=auto --iree-hal-target-backends=llvm-cpu --iree-llvmcpu-link-embedded=false --iree-llvmcpu-link-static --iree-llvmcpu-embedded-linker-path=./Binaries/ThirdParty/IREE/Mac/ld.lld --output-format=vm-bytecode --iree-llvmcpu-static-library-output-path=${OBJECT_PATH} -o ${VMFB_PATH} ${INPUT_PATH}",
+			"x86Features": [ "avx2", "fma", "f16c" ],
+			"CompilerArguments": "--iree-llvmcpu-target-triple=x86_64-apple-darwin --iree-input-type=auto --iree-hal-target-backends=llvm-cpu --iree-llvmcpu-target-cpu-features=+sse,+sse2,+ssse3,+sse4.1,+sse4.2,+avx,+f16c,+avx2,+fma --iree-llvmcpu-link-embedded=false --iree-llvmcpu-link-static --iree-llvmcpu-embedded-linker-path=./Binaries/ThirdParty/IREE/Mac/ld.lld --output-format=vm-bytecode --iree-llvmcpu-static-library-output-path=${OBJECT_PATH} -o ${VMFB_PATH} ${INPUT_PATH}",
 			"LinkerArguments": "-dylib -platform_version macos 11.0.0 11.0.0 -arch x86_64 -syslibroot \"${MAC_SDK_ROOT}\" -lSystem -o ${SHARED_LIB_PATH} ${OBJECT_PATH}"
 		},
 		{
-			"Architecture": "arm64",
-			"CompilerArguments": "--iree-llvmcpu-target-triple=arm64-apple-darwin --iree-input-type=auto --iree-hal-target-backends=llvm-cpu --iree-llvmcpu-link-embedded=false --iree-llvmcpu-link-static --iree-llvmcpu-embedded-linker-path=./Binaries/ThirdParty/IREE/Mac/ld.lld --output-format=vm-bytecode --iree-llvmcpu-static-library-output-path=${OBJECT_PATH} -o ${VMFB_PATH} ${INPUT_PATH}",
-			"LinkerArguments": "-dylib -platform_version macos 11.0.0 11.0.0 -arch arm64 -syslibroot \"${MAC_SDK_ROOT}\" -lSystem -o ${SHARED_LIB_PATH} ${OBJECT_PATH}"
+			"Name": "Intel_Default",
+			"Architecture": "x86_64",
+			"CompilerArguments": "--iree-llvmcpu-target-triple=x86_64-apple-darwin --iree-input-type=auto --iree-hal-target-backends=llvm-cpu --iree-llvmcpu-link-embedded=false --iree-llvmcpu-link-static --iree-llvmcpu-embedded-linker-path=./Binaries/ThirdParty/IREE/Mac/ld.lld --output-format=vm-bytecode --iree-llvmcpu-static-library-output-path=${OBJECT_PATH} -o ${VMFB_PATH} ${INPUT_PATH}",
+			"LinkerArguments": "-dylib -platform_version macos 11.0.0 11.0.0 -arch x86_64 -syslibroot \"${MAC_SDK_ROOT}\" -lSystem -o ${SHARED_LIB_PATH} ${OBJECT_PATH}"
 		}
 	]
 }
\ No newline at end of file
diff --git a/Engine/Plugins/Experimental/NNERuntimeIREE/Config/IREE_Windows_To_Windows.json b/Engine/Plugins/Experimental/NNERuntimeIREE/Config/IREE_Windows_To_Windows.json
--- a/Engine/Plugins/Experimental/NNERuntimeIREE/Config/IREE_Windows_To_Windows.json
+++ b/Engine/Plugins/Experimental/NNERuntimeIREE/Config/IREE_Windows_To_Windows.json
@@ -6,8 +6,30 @@
 	"SharedLibExt": ".dll",
 	"Targets": [
 		{
+			"Name": "AVX512",
 			"Architecture": "",
-			"CompilerArguments": "--iree-llvmcpu-target-triple=x86_64-pc-windows-msvc --iree-input-type=auto --iree-hal-target-backends=llvm-cpu --iree-llvmcpu-link-embedded=false --iree-llvmcpu-link-static --output-format=vm-bytecode --iree-llvmcpu-static-library-output-path=${OBJECT_PATH} -o ${VMFB_PATH} ${INPUT_PATH}",
+			"x86Features": [ "avx512", "fma", "f16c" ],
+			"CompilerArguments": "--iree-llvmcpu-target-triple=x86_64-pc-windows-msvc --iree-input-type=auto --iree-hal-target-backends=llvm-cpu --iree-llvmcpu-target-cpu-features=+sse,+sse2,+ssse3,+sse4.1,+sse4.2,+avx,+f16c,+avx2,+fma,+avx512f,+avx512vl,+avx512bw,+avx512dq --iree-llvmcpu-link-embedded=false --iree-llvmcpu-link-static --output-format=vm-bytecode --iree-llvmcpu-static-library-output-path=${OBJECT_PATH} -o ${VMFB_PATH} ${INPUT_PATH}",
+			"LinkerArguments": "-shared -o ${SHARED_LIB_PATH} ${OBJECT_PATH}"
+		},
+		{
+			"Name": "AVX2",
+			"Architecture": "",
+			"x86Features": [ "avx2", "fma", "f16c" ],
+			"CompilerArguments": "--iree-llvmcpu-target-triple=x86_64-pc-windows-msvc --iree-input-type=auto --iree-hal-target-backends=llvm-cpu --iree-llvmcpu-target-cpu-features=+sse,+sse2,+ssse3,+sse4.1,+sse4.2,+avx,+f16c,+avx2,+fma --iree-llvmcpu-link-embedded=false --iree-llvmcpu-link-static --output-format=vm-bytecode --iree-llvmcpu-static-library-output-path=${OBJECT_PATH} -o ${VMFB_PATH} ${INPUT_PATH}",
+			"LinkerArguments": "-shared -o ${SHARED_LIB_PATH} ${OBJECT_PATH}"
+		},
+		{
+			"Name": "AVX",
+			"Architecture": "",
+			"x86Features": [ "avx", "f16c" ],
+			"CompilerArguments": "--iree-llvmcpu-target-triple=x86_64-pc-windows-msvc --iree-input-type=auto --iree-hal-target-backends=llvm-cpu --iree-llvmcpu-target-cpu-features=+sse,+sse2,+ssse3,+sse4.1,+sse4.2,+avx,+f16c --iree-llvmcpu-link-embedded=false --iree-llvmcpu-link-static --output-format=vm-bytecode --iree-llvmcpu-static-library-output-path=${OBJECT_PATH} -o ${VMFB_PATH} ${INPUT_PATH}",
+			"LinkerArguments": "-shared -o ${SHARED_LIB_PATH} ${OBJECT_PATH}"
+		},
+		{
+			"Name": "Default",
+			"Architecture": "",
+			"CompilerArguments": "--iree-llvmcpu-target-triple=x86_64-pc-windows-msvc --iree-input-type=auto --iree-hal-target-backends=llvm-cpu --iree-llvmcpu-target-cpu-features=+sse,+sse2,+ssse3,+sse4.1,+sse4.2 --iree-llvmcpu-link-embedded=false --iree-llvmcpu-link-static --output-format=vm-bytecode --iree-llvmcpu-static-library-output-path=${OBJECT_PATH} -o ${VMFB_PATH} ${INPUT_PATH}",
 			"LinkerArguments": "-shared -o ${SHARED_LIB_PATH} ${OBJECT_PATH}"
 		}
 	]
diff --git a/Engine/Plugins/Experimental/NNERuntimeIREE/Source/NNERuntimeIREE/Private/NNERuntimeIREE.cpp b/Engine/Plugins/Experimental/NNERuntimeIREE/Source/NNERuntimeIREE/Private/NNERuntimeIREE.cpp
--- a/Engine/Plugins/Experimental/NNERuntimeIREE/Source/NNERuntimeIREE/Private/NNERuntimeIREE.cpp
+++ b/Engine/Plugins/Experimental/NNERuntimeIREE/Source/NNERuntimeIREE/Private/NNERuntimeIREE.cpp
@@ -12,6 +12,7 @@
 #include "NNERuntimeIREEMetaData.h"
 #endif // WITH_EDITOR
 
+#include "GenericPlatform/GenericPlatformMisc.h"
 #include "HAL/Platform.h"
 #include "Interfaces/ITargetPlatform.h"
 #include "IREECompilerRDG.h"
@@ -133,6 +134,79 @@ namespace UE::NNERuntimeIREE::Private
 
 		return ShaderPlatforms;
 	}
+
+	bool DoesArchitectureMatch(const FString& InArchitecture)
+	{
+		if (InArchitecture.IsEmpty())
+		{
+			return true;
+		}
+
+#if PLATFORM_CPU_X86_FAMILY
+		return InArchitecture.Equals(TEXT("x86_64"));
+#elif PLATFORM_CPU_ARM_FAMILY
+		return InArchitecture.Equals(TEXT("arm64"));
+#else
+		return false;
+#endif
+	}
+
+#if PLATFORM_CPU_X86_FAMILY
+	uint32 ParseX86Feature(const FString& InFeature)
+	{
+		if (InFeature.Equals(TEXT("sse4.2")))
+		{
+			return ECPUFeatureBits_X86::SSE42;
+		}
+		if (InFeature.Equals(TEXT("avx")))
+		{
+			return ECPUFeatureBits_X86::AVX;
+		}
+		if (InFeature.Equals(TEXT("avx2")))
+		{
+			return ECPUFeatureBits_X86::AVX2;
+		}
+		if (InFeature.Equals(TEXT("avx512")))
+		{
+			return ECPUFeatureBits_X86::AVX512_NOCAVEATS;
+		}
+		if (InFeature.Equals(TEXT("f16c")))
+		{
+			return ECPUFeatureBits_X86::F16C;
+		}
+		if (InFeature.Equals(TEXT("fma")))
+		{
+			return ECPUFeatureBits_X86::FMA;
+		}
+		return 0;
+	}
+
+	bool DoX86FeaturesMatch(const TConstArrayView<FString>& InFeatures)
+	{
+		for (const FString& Feature : InFeatures)
+		{
+			const uint32 FeatureBits = ParseX86Feature(Feature);
+			if (FeatureBits)
+			{
+				if (!FGenericPlatformMisc::CheckFeatureBit_X86(FeatureBits))
+				{
+					return false;
+				}
+			}
+			else
+			{
+				UE_LOG(LogNNERuntimeIREE, Warning, TEXT("NNERuntimeIREE is going to ignore unknown x86 feature %s."), *Feature);
+			}
+		}
+		return true;
+	}
+#else
+	bool DoX86FeaturesMatch(const TArray<FString>& InFeatures)
+	{
+		return InFeatures.IsEmpty();
+	}
+#endif
+
 } // UE::NNERuntimeIREE::Private
 
 FGuid UNNERuntimeIREECpu::GUID = FGuid((int32)'I', (int32)'C', (int32)'P', (int32)'U');
@@ -296,7 +370,7 @@ TSharedPtr<UE::NNE::FSharedModelData> UNNERuntimeIREECpu::CreateModelData(const
 
 			const FNNERuntimeIREEArchitectureInfoCPU& Info = CompilerResult.ArchitectureInfos[i];
 			const FString SharedLibrarySubDirFullPath = FPaths::Combine(SharedLibraryDirFullPath, Info.RelativeDirPath);
-			const FString StagingSubDirFullPath = FPaths::Combine(StagingDirFullPath, Info.Architecture);
+			const FString StagingSubDirFullPath = FPaths::Combine(StagingDirFullPath, Info.RelativeDirPath);
 			
 			const FString SharedLibraryFilePathSrc = FPaths::Combine(SharedLibrarySubDirFullPath, Info.SharedLibraryFileName);
 			const FString VmfbFilePathSrc = FPaths::Combine(SharedLibrarySubDirFullPath, Info.VmfbFileName);
@@ -352,13 +426,6 @@ TSharedPtr<UE::NNE::IModelCPU> UNNERuntimeIREECpu::CreateModelCPU(const TObjectP
 		return TSharedPtr<UE::NNE::IModelCPU>();
 	}
 
-	FString CurrentArchitecture = "";
-#if PLATFORM_CPU_X86_FAMILY
-	CurrentArchitecture = "x86_64";
-#elif PLATFORM_CPU_ARM_FAMILY
-	CurrentArchitecture = "arm64";
-#endif
-
 	TSharedPtr<UE::NNE::FSharedModelData> SharedData = ModelData->GetModelData(GetRuntimeName());
 	check(SharedData.IsValid());
 
@@ -394,22 +461,28 @@ TSharedPtr<UE::NNE::IModelCPU> UNNERuntimeIREECpu::CreateModelCPU(const TObjectP
 		FNNERuntimeIREECompilerResultCPU::StaticStruct()->SerializeBin(Reader, &CompilerResult);
 	}
 
-	int32 ArchitectureIndex = -1;
+	int32 TargetIndex = -1;
 	for (int32 i = 0; i < CompilerResult.ArchitectureInfos.Num(); i++)
 	{
-		if ((CompilerResult.ArchitectureInfos[i].Architecture.IsEmpty() && ArchitectureIndex < 0) || CompilerResult.ArchitectureInfos[i].Architecture.Equals(CurrentArchitecture))
+		bool bMatches = true;
+		bMatches &= DoesArchitectureMatch(CompilerResult.ArchitectureInfos[i].Architecture);
+		bMatches &= DoX86FeaturesMatch(CompilerResult.ArchitectureInfos[i].X86Features);
+
+		if (bMatches)
 		{
-			ArchitectureIndex = i;
+			TargetIndex = i;
+			break;
 		}
 	}
 
-	if (ArchitectureIndex < 0)
+	if (TargetIndex < 0)
 	{
-		UE_LOG(LogNNERuntimeIREE, Warning, TEXT("UNNERuntimeIREECpu failed to find a matching architecture for \'%s\'"), *CurrentArchitecture);
+		UE_LOG(LogNNERuntimeIREE, Warning, TEXT("UNNERuntimeIREECpu failed to find a matching target for the current hardware"));
 		return TSharedPtr<UE::NNE::IModelCPU>();
 	}
 
-	const FNNERuntimeIREEArchitectureInfoCPU& ArchitectureInfo = CompilerResult.ArchitectureInfos[ArchitectureIndex];
+
+	const FNNERuntimeIREEArchitectureInfoCPU& ArchitectureInfo = CompilerResult.ArchitectureInfos[TargetIndex];
 
 	const FString FileIdString = IREEModelData->FileId.ToString(EGuidFormats::Digits).ToLower();
 	const FString SharedLibraryDirFullPath = FPaths::ConvertRelativePathToFull(FPaths::Combine(FPaths::ProjectDir(), GetSharedLibDirPath(UGameplayStatics::GetPlatformName(), GetRuntimeName(), FileIdString)));
@@ -421,6 +494,7 @@ TSharedPtr<UE::NNE::IModelCPU> UNNERuntimeIREECpu::CreateModelCPU(const TObjectP
 		UE_LOG(LogNNERuntimeIREE, Warning, TEXT("UNNERuntimeIREECpu could not initialize the model created from model data with id %s"), *FileIdString);
 		return TSharedPtr<UE::NNE::IModelCPU>();
 	}
+	UE_LOG(LogNNERuntimeIREE, Display, TEXT("UNNERuntimeIREECpu successfully loaded target %s"), *ArchitectureInfo.RelativeDirPath);
 
 	return Model;
 }
diff --git a/Engine/Plugins/Experimental/NNERuntimeIREE/Source/NNERuntimeIREE/Private/NNERuntimeIREECompiler.cpp b/Engine/Plugins/Experimental/NNERuntimeIREE/Source/NNERuntimeIREE/Private/NNERuntimeIREECompiler.cpp
--- a/Engine/Plugins/Experimental/NNERuntimeIREE/Source/NNERuntimeIREE/Private/NNERuntimeIREECompiler.cpp
+++ b/Engine/Plugins/Experimental/NNERuntimeIREE/Source/NNERuntimeIREE/Private/NNERuntimeIREECompiler.cpp
@@ -218,9 +218,28 @@ namespace UE::NNERuntimeIREE
 			const FString PluginDir = FPaths::ConvertRelativePathToFull(*IPluginManager::Get().FindPlugin(UE_PLUGIN_NAME)->GetBaseDir());
 
 			bool bResult = true;
+			OutCompilerResult.Reset();
+			TArray<FString> PreviousRelativeDirPaths;
 			for (int32 i = 0; i < BuildTargets.Num(); i++)
 			{
-				FString IntermediateDirPath = FPaths::Combine(InOutputDir, BuildTargets[i].Architecture);
+				FString RelativeDirPath;
+				if (BuildTargets.Num() > 1)
+				{
+					if (!BuildTargets[i].Name.IsEmpty())
+					{
+						RelativeDirPath = FPaths::MakeValidFileName(BuildTargets[i].Name);
+						if (RelativeDirPath.StartsWith(TEXT("Target_")))
+						{
+							RelativeDirPath.Empty();
+						}
+					}
+					if (RelativeDirPath.IsEmpty() || PreviousRelativeDirPaths.Contains(RelativeDirPath))
+					{
+						RelativeDirPath = FString::Printf(TEXT("Target_%d"), i);
+					}
+					PreviousRelativeDirPaths.Add(RelativeDirPath);
+				}
+				FString IntermediateDirPath = FPaths::Combine(InOutputDir, RelativeDirPath);
 				PlatformFile.CreateDirectoryTree(*IntermediateDirPath);
 				FString IntermediateFilePathNoExt = FPaths::Combine(IntermediateDirPath, InModelName);
 				FString ObjectFilePath = IntermediateFilePathNoExt + ".o";
@@ -232,13 +251,13 @@ namespace UE::NNERuntimeIREE
 				{
 					UE_LOG(LogNNERuntimeIREE, Warning, TEXT("UNNERuntimeIREECpu could not replace environment variables in %s"), *BuildTargets[i].CompilerArguments);
 					bResult = false;
-					continue;
+					break;
 				}
 				if (!IREEUtils::ResolveSdkPaths(CompilerArguments))
 				{
 					UE_LOG(LogNNERuntimeIREE, Warning, TEXT("UNNERuntimeIREECpu could not replace SDK paths in %s"), *BuildTargets[i].CompilerArguments);
 					bResult = false;
-					continue;
+					break;
 				}
 				CompilerArguments.ReplaceInline(*FString("${OBJECT_PATH}"), *(FString("\"") + ObjectFilePath + "\""));
 				CompilerArguments.ReplaceInline(*FString("${VMFB_PATH}"), *(FString("\"") + VmfbFilePath + "\""));
@@ -255,7 +274,7 @@ namespace UE::NNERuntimeIREE
 					UE_LOG(LogNNERuntimeIREE, Warning, TEXT("UNNERuntimeIREECpu failed to compile the model \"%s\" using the command:"), *InputFilePath);
 					UE_LOG(LogNNERuntimeIREE, Warning, TEXT("\"%s\" %s"), *CompilerCommand, *CompilerArguments);
 					bResult = false;
-					continue;
+					break;
 				}
 
 				FString LinkerArguments = BuildTargets[i].LinkerArguments;
@@ -263,13 +282,13 @@ namespace UE::NNERuntimeIREE
 				{
 					UE_LOG(LogNNERuntimeIREE, Warning, TEXT("UNNERuntimeIREECpu could not replace environment variables in %s"), *BuildTargets[i].LinkerArguments);
 					bResult = false;
-					continue;
+					break;
 				}
 				if (!IREEUtils::ResolveSdkPaths(LinkerArguments))
 				{
 					UE_LOG(LogNNERuntimeIREE, Warning, TEXT("UNNERuntimeIREECpu could not replace SDK paths in %s"), *BuildTargets[i].LinkerArguments);
 					bResult = false;
-					continue;
+					break;
 				}
 				LinkerArguments.ReplaceInline(*FString("${OBJECT_PATH}"), *(FString("\"") + ObjectFilePath + "\""));
 				LinkerArguments.ReplaceInline(*FString("${SHARED_LIB_PATH}"), *(FString("\"") + SharedLibFilePath + "\""));
@@ -285,7 +304,7 @@ namespace UE::NNERuntimeIREE
 					UE_LOG(LogNNERuntimeIREE, Warning, TEXT("UNNERuntimeIREECpu failed to link the model \"%s\" using the command:"), *InputFilePath);
 					UE_LOG(LogNNERuntimeIREE, Warning, TEXT("\"%s\" %s"), *LinkerCommand, *LinkerArguments);
 					bResult = false;
-					continue;
+					break;
 				}
 
 				FString SharedLibraryEntryPointName = "";
@@ -294,26 +313,27 @@ namespace UE::NNERuntimeIREE
 				{
 					UE_LOG(LogNNERuntimeIREE, Warning, TEXT("UNNERuntimeIREECpu could not find the model header \"%s\""), *HeaderPath);
 					bResult = false;
-					continue;
+					break;
 				}
 				FString HeaderString;
 				if (!FFileHelper::LoadFileToString(HeaderString, *HeaderPath) || HeaderString.IsEmpty())
 				{
 					UE_LOG(LogNNERuntimeIREE, Warning, TEXT("UNNERuntimeIREECpu could not read the model header \"%s\""), *HeaderPath);
 					bResult = false;
-					continue;
+					break;
 				}
 				SharedLibraryEntryPointName = GetSharedLibraryEntryPointName(HeaderString);
 				if (SharedLibraryEntryPointName.IsEmpty())
 				{
 					UE_LOG(LogNNERuntimeIREE, Warning, TEXT("UNNERuntimeIREECpu could not find the entry point in model header \"%s\""), *HeaderPath);
 					bResult = false;
-					continue;
+					break;
 				}
 
 				FNNERuntimeIREEArchitectureInfoCPU ArchitectureInfo;
 				ArchitectureInfo.Architecture = BuildTargets[i].Architecture;
-				ArchitectureInfo.RelativeDirPath = BuildTargets[i].Architecture;
+				ArchitectureInfo.X86Features = BuildTargets[i].X86Features;
+				ArchitectureInfo.RelativeDirPath = RelativeDirPath;
 				ArchitectureInfo.SharedLibraryFileName = InModelName + SharedLibExt;
 				ArchitectureInfo.VmfbFileName = InModelName + ".vmfb";
 				ArchitectureInfo.SharedLibraryEntryPointName = SharedLibraryEntryPointName;
@@ -324,7 +344,7 @@ namespace UE::NNERuntimeIREE
 
 			if (!bResult)
 			{
-				OutCompilerResult.ArchitectureInfos.Empty();
+				OutCompilerResult.Reset();
 			}
 
 			return bResult;
diff --git a/Engine/Plugins/Experimental/NNERuntimeIREE/Source/NNERuntimeIREE/Private/NNERuntimeIREECompiler.h b/Engine/Plugins/Experimental/NNERuntimeIREE/Source/NNERuntimeIREE/Private/NNERuntimeIREECompiler.h
--- a/Engine/Plugins/Experimental/NNERuntimeIREE/Source/NNERuntimeIREE/Private/NNERuntimeIREECompiler.h
+++ b/Engine/Plugins/Experimental/NNERuntimeIREE/Source/NNERuntimeIREE/Private/NNERuntimeIREECompiler.h
@@ -16,6 +16,9 @@ struct FNNERuntimeIREEArchitectureInfoCPU
 	UPROPERTY()
 	FString Architecture;
 
+	UPROPERTY()
+	TArray<FString> X86Features;
+
 	UPROPERTY()
 	FString RelativeDirPath;
 	
@@ -36,6 +39,11 @@ struct FNNERuntimeIREECompilerResultCPU
 
 	UPROPERTY()
 	TArray<FNNERuntimeIREEArchitectureInfoCPU> ArchitectureInfos;
+
+	void Reset()
+	{
+		ArchitectureInfos.Reset();
+	}
 };
 
 #ifdef WITH_NNE_RUNTIME_IREE
@@ -45,12 +53,16 @@ namespace UE::NNERuntimeIREE::CPU
 {
 	struct FBuildTarget : public FJsonSerializable
 	{
+		FString Name;
 		FString Architecture;
+		TArray<FString> X86Features;
 		FString CompilerArguments;
 		FString LinkerArguments;
 
 		BEGIN_JSON_SERIALIZER
+			JSON_SERIALIZE("Name", Name);
 			JSON_SERIALIZE("Architecture", Architecture);
+			JSON_SERIALIZE_ARRAY("X86Features", X86Features);
 			JSON_SERIALIZE("CompilerArguments", CompilerArguments);
 			JSON_SERIALIZE("LinkerArguments", LinkerArguments);
 		END_JSON_SERIALIZER
