[NNE] NNERuntimeIREE remove custom versions

#rnx
#tests NNEModelTest

Currently IREE uses custom versions to serialize it's model data. Since these serializations happen into a memory buffer and are not part of an asset, these custom version never get saved. Since this data is only stored to cache iree compilation we could simply use the manually stored iree version to recreated the cache each time there is a mismatch between the stored and the latest version. This has two benefits: It prevents bugs due to running old compilation in a new iree runtime and it simplifies serialization since we don't need to maintain backward compatibility.

#Synced-CL 49655169

diff --git a/Engine/Plugins/Experimental/NNERuntimeIREE/Source/NNERuntimeIREE/Private/NNERuntimeIREE.cpp b/Engine/Plugins/Experimental/NNERuntimeIREE/Source/NNERuntimeIREE/Private/NNERuntimeIREE.cpp
--- a/Engine/Plugins/Experimental/NNERuntimeIREE/Source/NNERuntimeIREE/Private/NNERuntimeIREE.cpp
+++ b/Engine/Plugins/Experimental/NNERuntimeIREE/Source/NNERuntimeIREE/Private/NNERuntimeIREE.cpp
@@ -267,35 +267,36 @@ TSharedPtr<UE::NNE::FSharedModelData> UNNERuntimeIREECpu::CreateModelData(const
 	FNNERuntimeIREECompilerResultCPU CompilerResult{};
 
 	bool bNeedCompileMlir = true;
-	if (PlatformFile.FileExists(*IREEModelDataFilePath))
+	if (PlatformFile.FileExists(*IREEModelDataFilePath) && FFileHelper::LoadFileToArray(ResultData, *IREEModelDataFilePath))
 	{
 		SCOPED_NAMED_EVENT_TEXT("Validate", FColor::Magenta);
 
-		FFileHelper::LoadFileToArray(ResultData, *IREEModelDataFilePath);
-		
+		if (UNNERuntimeIREEModelDataCPU::IsSameGuidAndVersion(ResultData, GUID, Version))
 		{
-			FMemoryReaderView Reader(ResultData, /*bIsPersitent =*/ true);
-			IREEModelData->Serialize(Reader);
-		}
+			{
+				FMemoryReaderView Reader(ResultData, /*bIsPersitent =*/ true);
+				IREEModelData->Serialize(Reader);
+			}
 
-		check(FileIdString.Equals(IREEModelData->FileId.ToString(EGuidFormats::Digits).ToLower()));
+			check(FileIdString.Equals(IREEModelData->FileId.ToString(EGuidFormats::Digits).ToLower()));
 
-		{
-			FMemoryReaderView Reader(IREEModelData->CompilerResult, /*bIsPersitent =*/ true);
-			FNNERuntimeIREECompilerResultCPU::StaticStruct()->SerializeBin(Reader, &CompilerResult);
-		}
+			{
+				FMemoryReaderView Reader(IREEModelData->CompilerResult, /*bIsPersitent =*/ true);
+				FNNERuntimeIREECompilerResultCPU::StaticStruct()->SerializeBin(Reader, &CompilerResult);
+			}
 
-		bNeedCompileMlir = false;
-		for (int32 i = 0; i < CompilerResult.ArchitectureInfos.Num(); i++)
-		{
-			const FNNERuntimeIREEArchitectureInfoCPU& Info = CompilerResult.ArchitectureInfos[i];
-			const FString SharedLibrarySubDirFullPath = FPaths::Combine(SharedLibraryDirFullPath, Info.RelativeDirPath);
+			bNeedCompileMlir = false;
+			for (int32 i = 0; i < CompilerResult.ArchitectureInfos.Num(); i++)
+			{
+				const FNNERuntimeIREEArchitectureInfoCPU& Info = CompilerResult.ArchitectureInfos[i];
+				const FString SharedLibrarySubDirFullPath = FPaths::Combine(SharedLibraryDirFullPath, Info.RelativeDirPath);
 
-			const FString SharedLibraryFilePath = FPaths::Combine(SharedLibrarySubDirFullPath, Info.SharedLibraryFileName);
-			const FString VmfbFilePath = FPaths::Combine(SharedLibrarySubDirFullPath, Info.VmfbFileName);
+				const FString SharedLibraryFilePath = FPaths::Combine(SharedLibrarySubDirFullPath, Info.SharedLibraryFileName);
+				const FString VmfbFilePath = FPaths::Combine(SharedLibrarySubDirFullPath, Info.VmfbFileName);
 
-			bNeedCompileMlir |= !PlatformFile.FileExists(*SharedLibraryFilePath);
-			bNeedCompileMlir |= !PlatformFile.FileExists(*VmfbFilePath);
+				bNeedCompileMlir |= !PlatformFile.FileExists(*SharedLibraryFilePath);
+				bNeedCompileMlir |= !PlatformFile.FileExists(*VmfbFilePath);
+			}
 		}
 	}
 	
@@ -693,8 +694,9 @@ TSharedPtr<UE::NNE::FSharedModelData> UNNERuntimeIREERdg::CreateModelData(const
 		SCOPED_NAMED_EVENT_TEXT("Validate", FColor::Magenta);
 
 		TUniquePtr<FArchive> FileReader(IFileManager::Get().CreateFileReader(*IREEModelDataFilePath));
-		if (FileReader)
+		if (FileReader && UNNERuntimeIREEModelDataRDG::IsSameGuidAndVersion(*FileReader, GUID, Version))
 		{
+			FileReader->Seek(0);
 			bNeedCompileMlir = false;
 
 			FNNERuntimeIREEModelDataHeaderRDG Header = UNNERuntimeIREEModelDataRDG::ReadHeader(*FileReader);
@@ -820,16 +822,12 @@ UNNERuntimeIREERdg::ECanCreateModelRDGStatus UNNERuntimeIREERdg::CanCreateModelR
 	}
 
 	FMemoryReaderView Reader(SharedData->GetView(), /*bIsPersitent =*/ true);
-	FNNERuntimeIREEModelDataHeaderRDG Header = UNNERuntimeIREEModelDataRDG::ReadHeader(Reader);
-
-	if (Header.GUID != GUID)
-	{
-		return ECanCreateModelRDGStatus::Fail;
-	}
-	if (Header.Version != Version)
+	if (!UNNERuntimeIREEModelDataRDG::IsSameGuidAndVersion(Reader, GUID, Version))
 	{
 		return ECanCreateModelRDGStatus::Fail;
 	}
+	Reader.Seek(0);
+	FNNERuntimeIREEModelDataHeaderRDG Header = UNNERuntimeIREEModelDataRDG::ReadHeader(Reader);
 
 	const FString ShaderPlatformName = LexToString(GShaderPlatformForFeatureLevel[GMaxRHIFeatureLevel]);
 	if (Header.ShaderPlatforms.Find(ShaderPlatformName) == INDEX_NONE)
diff --git a/Engine/Plugins/Experimental/NNERuntimeIREE/Source/NNERuntimeIREE/Private/NNERuntimeIREEModelData.cpp b/Engine/Plugins/Experimental/NNERuntimeIREE/Source/NNERuntimeIREE/Private/NNERuntimeIREEModelData.cpp
--- a/Engine/Plugins/Experimental/NNERuntimeIREE/Source/NNERuntimeIREE/Private/NNERuntimeIREEModelData.cpp
+++ b/Engine/Plugins/Experimental/NNERuntimeIREE/Source/NNERuntimeIREE/Private/NNERuntimeIREEModelData.cpp
@@ -4,55 +4,14 @@
 
 #include "Containers/Array.h"
 #include "NNE.h"
-#include "NNERuntimeIREELog.h"
-#include "Serialization/CustomVersion.h"
-
-namespace UE::NNERuntimeIREE::CPU::ModelData::Private
-{
-
-enum Version : uint32
-{
-	V0 = 0, // Initial
-	// New versions can be added above this line
-	VersionPlusOne,
-	Latest = VersionPlusOne - 1
-};
-const FGuid GUID(0x6dcb835d, 0x9ac64a1d, 0x8165d871, 0x6122dab7);
-FCustomVersionRegistration Version(GUID, Version::Latest, TEXT("NNERuntimeIREEModelDataVersionCPU"));// Always save with the latest version
-
-} // UE::NNERuntimeIREE::CPU::ModelData::Private
 
 void UNNERuntimeIREEModelDataCPU::Serialize(FArchive& Ar)
 {
-	// Store the asset version (no effect in load)
-	Ar.UsingCustomVersion(UE::NNERuntimeIREE::CPU::ModelData::Private::GUID);
-
-	if (Ar.IsSaving() || Ar.IsCountingMemory()) 
-	{
-		Ar << GUID;
-		Ar << Version;
-		Ar << FileId;
-		Ar << ModuleMetaData;
-		Ar << CompilerResult;
-	}
-	else
-	{
-		int32 NumArchitectures = 0;
-
-		switch (Ar.CustomVer(UE::NNERuntimeIREE::CPU::ModelData::Private::GUID))
-		{
-		case UE::NNERuntimeIREE::CPU::ModelData::Private::Version::V0:
-			Ar << GUID;
-			Ar << Version;
-			Ar << FileId;
-			Ar << ModuleMetaData;
-			Ar << CompilerResult;
-			break;
-		default:
-			UE_LOG(LogNNERuntimeIREE, Error, TEXT("UNNERuntimeIREEModelDataCPU: Unknown asset version %d: Deserialisation failed, please reimport the original model."), Ar.CustomVer(UE::NNERuntimeIREE::CPU::ModelData::Private::GUID));
-			break;
-		}
-	}
+	Ar << GUID;
+	Ar << Version;
+	Ar << FileId;
+	Ar << ModuleMetaData;
+	Ar << CompilerResult;
 }
 
 bool UNNERuntimeIREEModelDataCPU::IsSameGuidAndVersion(TConstArrayView64<uint8> Data, FGuid Guid, int32 Version)
@@ -70,60 +29,22 @@ bool UNNERuntimeIREEModelDataCPU::IsSameGuidAndVersion(TConstArrayView64<uint8>
 	return bResult;
 }
 
-namespace UE::NNERuntimeIREE::RDG::ModelData::Private
-{
-
-enum Version : uint32
-{
-	V0 = 0, // Initial
-	// New versions can be added above this line
-	VersionPlusOne,
-	Latest = VersionPlusOne - 1
-};
-const FGuid GUID(0x29663d20, 0x1dbd43c5, 0xa1ecc03b, 0x6005c845);
-FCustomVersionRegistration Version(GUID, Version::Latest, TEXT("NNERuntimeIREEModelDataVersionRDG"));// Always save with the latest version
-
-} // UE::NNERuntimeIREE::RDG::ModelData::Private
-
 void UNNERuntimeIREEModelDataRDG::Serialize(FArchive& Ar)
 {
-	// Store the asset version (no effect in load)
-	Ar.UsingCustomVersion(UE::NNERuntimeIREE::RDG::ModelData::Private::GUID);
-
-	if (Ar.IsSaving() || Ar.IsCountingMemory()) 
-	{
-		FNNERuntimeIREEModelDataHeaderRDG::StaticStruct()->SerializeBin(Ar, &Header);
-		Ar << ModuleMetaData;
-		Ar << CompilerResult;
-	}
-	else
-	{
-		switch (Ar.CustomVer(UE::NNERuntimeIREE::RDG::ModelData::Private::GUID))
-		{
-		case UE::NNERuntimeIREE::RDG::ModelData::Private::Version::V0:
-		FNNERuntimeIREEModelDataHeaderRDG::StaticStruct()->SerializeBin(Ar, &Header);
-			Ar << ModuleMetaData;
-			Ar << CompilerResult;
-			break;
-		default:
-			UE_LOG(LogNNERuntimeIREE, Error, TEXT("UNNERuntimeIREEModelDataRDG: Unknown asset version %d: Deserialisation failed, please reimport the original model."), Ar.CustomVer(UE::NNERuntimeIREE::RDG::ModelData::Private::GUID));
-			break;
-		}
-	}
+	FNNERuntimeIREEModelDataHeaderRDG::StaticStruct()->SerializeBin(Ar, &Header);
+	Ar << ModuleMetaData;
+	Ar << CompilerResult;
 }
 
 FNNERuntimeIREEModelDataHeaderRDG UNNERuntimeIREEModelDataRDG::ReadHeader(FArchive& Ar)
 {
 	FNNERuntimeIREEModelDataHeaderRDG Header{};
-	switch (Ar.CustomVer(UE::NNERuntimeIREE::RDG::ModelData::Private::GUID))
-	{
-	case UE::NNERuntimeIREE::RDG::ModelData::Private::Version::V0:
 	FNNERuntimeIREEModelDataHeaderRDG::StaticStruct()->SerializeBin(Ar, &Header);
-		break;
-	default:
-		UE_LOG(LogNNERuntimeIREE, Error, TEXT("UNNERuntimeIREEModelDataRDG: Unknown asset version %d: Deserialisation failed, please reimport the original model."), Ar.CustomVer(UE::NNERuntimeIREE::RDG::ModelData::Private::GUID));
-		break;
-	}
-
 	return Header;
+}
+
+bool UNNERuntimeIREEModelDataRDG::IsSameGuidAndVersion(FArchive& Ar, FGuid Guid, int32 Version)
+{
+	FNNERuntimeIREEModelDataHeaderRDG Header = ReadHeader(Ar);
+	return Guid == Header.GUID && Version == Header.Version && !Ar.IsError();
 }
\ No newline at end of file
diff --git a/Engine/Plugins/Experimental/NNERuntimeIREE/Source/NNERuntimeIREE/Private/NNERuntimeIREEModelData.h b/Engine/Plugins/Experimental/NNERuntimeIREE/Source/NNERuntimeIREE/Private/NNERuntimeIREEModelData.h
--- a/Engine/Plugins/Experimental/NNERuntimeIREE/Source/NNERuntimeIREE/Private/NNERuntimeIREEModelData.h
+++ b/Engine/Plugins/Experimental/NNERuntimeIREE/Source/NNERuntimeIREE/Private/NNERuntimeIREEModelData.h
@@ -99,6 +99,11 @@ public:
 	virtual void Serialize(FArchive& Ar) override;
 	// End of UObject interface
 
+	/**
+	 * Check data for matching Guid and Version with only deserializing the header
+	 */
+	static bool IsSameGuidAndVersion(FArchive& Ar, FGuid Guid, int32 Version);
+
 	static FNNERuntimeIREEModelDataHeaderRDG ReadHeader(FArchive& Ar);
 
 	/**
