[NNE] NNERuntimeIREE fixes in UNNERuntimeIREEModuleMetaData::ParseFromBuffer

#rnx
#tests NNEModelTest

- return false during failure
- make use of this return value
- reset FunctionMetaData in order to prevent issues with repeated calls to this function

#Synced-CL 49725404

diff --git a/Engine/Plugins/Experimental/NNERuntimeIREE/Source/NNERuntimeIREE/Private/NNERuntimeIREE.cpp b/Engine/Plugins/Experimental/NNERuntimeIREE/Source/NNERuntimeIREE/Private/NNERuntimeIREE.cpp
--- a/Engine/Plugins/Experimental/NNERuntimeIREE/Source/NNERuntimeIREE/Private/NNERuntimeIREE.cpp
+++ b/Engine/Plugins/Experimental/NNERuntimeIREE/Source/NNERuntimeIREE/Private/NNERuntimeIREE.cpp
@@ -327,9 +327,8 @@ TSharedPtr<UE::NNE::FSharedModelData> UNNERuntimeIREECpu::CreateModelData(const
 
 		// NOTE: From this point FileData refers to mlir data in any case
 		
-		TWeakObjectPtr<UNNERuntimeIREEModuleMetaData> CompilerModuleMetaData = NewObject<UNNERuntimeIREEModuleMetaData>();
 
-		if (!Compiler->CompileMlir(FileData, FileIdString, IntermediateDirFullPath, CompilerResult, *CompilerModuleMetaData))
+		if (!Compiler->CompileMlir(FileData, FileIdString, IntermediateDirFullPath, CompilerResult))
 		{
 			UE_LOG(LogNNERuntimeIREE, Warning, TEXT("UNNERuntimeIREECpu failed to compile model %s"), *FileIdString);
 			return TSharedPtr<UE::NNE::FSharedModelData>();
@@ -344,6 +343,12 @@ TSharedPtr<UE::NNE::FSharedModelData> UNNERuntimeIREECpu::CreateModelData(const
 		}
 		if (IREEModelData->ModuleMetaData.IsEmpty())
 		{
+			TWeakObjectPtr<UNNERuntimeIREEModuleMetaData> CompilerModuleMetaData = NewObject<UNNERuntimeIREEModuleMetaData>();
+			if (!CompilerModuleMetaData->ParseFromBuffer(FileData))
+			{
+				UE_LOG(LogNNERuntimeIREE, Warning, TEXT("UNNERuntimeIREECpu failed to parse meta data %s"), *FileIdString);
+				return TSharedPtr<UE::NNE::FSharedModelData>();
+			}
 			FMemoryWriter64 Writer(IREEModelData->ModuleMetaData, /*bIsPersitent =*/ true);
 			CompilerModuleMetaData->Serialize(Writer);
 		}
@@ -769,7 +774,11 @@ TSharedPtr<UE::NNE::FSharedModelData> UNNERuntimeIREERdg::CreateModelData(const
 		{
 			TWeakObjectPtr<UNNERuntimeIREEModuleMetaData> CompilerModuleMetaData = NewObject<UNNERuntimeIREEModuleMetaData>();
 
-			CompilerModuleMetaData->ParseFromBuffer(FileData);
+			if (!CompilerModuleMetaData->ParseFromBuffer(FileData))
+			{
+				UE_LOG(LogNNERuntimeIREE, Warning, TEXT("UNNERuntimeIREERdg failed to parse meta data %s"), *FileIdString);
+				return TSharedPtr<UE::NNE::FSharedModelData>();
+			}
 
 			FMemoryWriter64 Writer(IREEModelData->ModuleMetaData, /*bIsPersitent =*/ true);
 			CompilerModuleMetaData->Serialize(Writer);
diff --git a/Engine/Plugins/Experimental/NNERuntimeIREE/Source/NNERuntimeIREE/Private/NNERuntimeIREECompiler.cpp b/Engine/Plugins/Experimental/NNERuntimeIREE/Source/NNERuntimeIREE/Private/NNERuntimeIREECompiler.cpp
--- a/Engine/Plugins/Experimental/NNERuntimeIREE/Source/NNERuntimeIREE/Private/NNERuntimeIREECompiler.cpp
+++ b/Engine/Plugins/Experimental/NNERuntimeIREE/Source/NNERuntimeIREE/Private/NNERuntimeIREECompiler.cpp
@@ -193,18 +193,12 @@ namespace UE::NNERuntimeIREE
 			return IREEUtils::ImportOnnx(ImporterCommand, ImporterArguments, InFileData, InModelName, InOutputDir, OutMlirData);
 		}
 
-		bool FCompiler::CompileMlir(TConstArrayView<uint8> InFileData, const FString& InModelName, const FString& InOutputDir, FNNERuntimeIREECompilerResultCPU& OutCompilerResult, UNNERuntimeIREEModuleMetaData& ModuleMetaData)
+		bool FCompiler::CompileMlir(TConstArrayView<uint8> InFileData, const FString& InModelName, const FString& InOutputDir, FNNERuntimeIREECompilerResultCPU& OutCompilerResult)
 		{
 			SCOPED_NAMED_EVENT_TEXT("FCompiler::CompileMlir", FColor::Magenta);
 
 			using namespace Private;
 
-			{
-				SCOPED_NAMED_EVENT_TEXT("Metadata", FColor::Magenta);
-
-				ModuleMetaData.ParseFromBuffer(InFileData);
-			}
-
 			IPlatformFile& PlatformFile = FPlatformFileManager::Get().GetPlatformFile();
 
 			FString InputFilePath = FPaths::Combine(InOutputDir, InModelName) + ".mlir";
diff --git a/Engine/Plugins/Experimental/NNERuntimeIREE/Source/NNERuntimeIREE/Private/NNERuntimeIREECompiler.h b/Engine/Plugins/Experimental/NNERuntimeIREE/Source/NNERuntimeIREE/Private/NNERuntimeIREECompiler.h
--- a/Engine/Plugins/Experimental/NNERuntimeIREE/Source/NNERuntimeIREE/Private/NNERuntimeIREECompiler.h
+++ b/Engine/Plugins/Experimental/NNERuntimeIREE/Source/NNERuntimeIREE/Private/NNERuntimeIREECompiler.h
@@ -3,7 +3,6 @@
 #pragma once
 
 #include "CoreMinimal.h"
-#include "NNERuntimeIREEMetaData.h"
 #include "Serialization/JsonSerializerMacros.h"
 
 #include "NNERuntimeIREECompiler.generated.h"
@@ -98,7 +97,7 @@ namespace UE::NNERuntimeIREE::CPU
 
 	public:
 		bool ImportOnnx(TConstArrayView<uint8> InFileData, const FString& InModelName, const FString& InOutputDir, TArray64<uint8>& OutMlirData);
-		bool CompileMlir(TConstArrayView<uint8> InFileData, const FString& InModelName, const FString& InOutputDir, FNNERuntimeIREECompilerResultCPU& OutCompilerResult, UNNERuntimeIREEModuleMetaData& ModuleMetaData);
+		bool CompileMlir(TConstArrayView<uint8> InFileData, const FString& InModelName, const FString& InOutputDir, FNNERuntimeIREECompilerResultCPU& OutCompilerResult);
 
 	private:
 		FString ImporterCommand;
diff --git a/Engine/Plugins/Experimental/NNERuntimeIREE/Source/NNERuntimeIREE/Private/NNERuntimeIREEMetaData.cpp b/Engine/Plugins/Experimental/NNERuntimeIREE/Source/NNERuntimeIREE/Private/NNERuntimeIREEMetaData.cpp
--- a/Engine/Plugins/Experimental/NNERuntimeIREE/Source/NNERuntimeIREE/Private/NNERuntimeIREEMetaData.cpp
+++ b/Engine/Plugins/Experimental/NNERuntimeIREE/Source/NNERuntimeIREE/Private/NNERuntimeIREEMetaData.cpp
@@ -254,7 +254,7 @@ bool UNNERuntimeIREEModuleMetaData::ParseFromBuffer(TConstArrayView64<uint8> Buf
 		};
 
 		const int32 NumFun = Mod.GetFunctionCount();
-		FunctionMetaData.Reserve(NumFun);
+		FunctionMetaData.Reset(NumFun);
 
 		for (int32 i = 0; i < NumFun; i++)
 		{
@@ -293,10 +293,12 @@ bool UNNERuntimeIREEModuleMetaData::ParseFromBuffer(TConstArrayView64<uint8> Buf
 	catch (const std::exception& Exception)
 	{
 		UE_LOG(LogNNERuntimeIREE, Error, TEXT("Parse MLIR exception: %hs"), Exception.what());
+		return false;
 	}
 	catch (...)
 	{
 		UE_LOG(LogNNERuntimeIREE, Error, TEXT("Parse MLIR unknown exception."));
+		return false;
 	}
 #endif // WITH_EDITOR
 
