[NNE] NNERuntimeIREE change ModuleMetaData to a simple class

#rnx
#tests NNEModelTest, manual test that multi functions are read correctly

Renames UNNERuntimeIREEModuleMetaData to FModuleMetaData, removes UObject inheritence and moves it into the UE::NNERuntimeIREE namespace.

#Synced-CL 49726700

diff --git a/Engine/Plugins/Experimental/NNERuntimeIREE/Source/NNERuntimeIREE/Internal/NNERuntimeIREEMetaData.h b/Engine/Plugins/Experimental/NNERuntimeIREE/Source/NNERuntimeIREE/Internal/NNERuntimeIREEMetaData.h
--- a/Engine/Plugins/Experimental/NNERuntimeIREE/Source/NNERuntimeIREE/Internal/NNERuntimeIREEMetaData.h
+++ b/Engine/Plugins/Experimental/NNERuntimeIREE/Source/NNERuntimeIREE/Internal/NNERuntimeIREEMetaData.h
@@ -4,9 +4,6 @@
 
 #include "Containers/ContainersFwd.h"
 #include "NNETypes.h"
-#include "UObject/Class.h"
-
-#include "NNERuntimeIREEMetaData.generated.h"
 
 namespace UE::NNERuntimeIREE
 {
@@ -16,20 +13,15 @@ namespace UE::NNERuntimeIREE
 		TArray<UE::NNE::FTensorDesc> InputDescs;
 		TArray<UE::NNE::FTensorDesc> OutputDescs;
 	};
-} // UE::NNERuntimeIREE
 
+	class NNERUNTIMEIREE_API FModuleMetaData
+	{
 
-UCLASS()
-class NNERUNTIMEIREE_API UNNERuntimeIREEModuleMetaData : public UObject
-{
-	GENERATED_BODY()
-
-public:
-	// UObject interface
-	virtual void Serialize(FArchive& Ar) override;
-	// End of UObject interface
+	public:
+		void Serialize(FArchive& Ar);
 
-	bool ParseFromBuffer(TConstArrayView64<uint8> Buffer);
+		bool ParseFromBuffer(TConstArrayView64<uint8> Buffer);
 
-	TArray<UE::NNERuntimeIREE::FFunctionMetaData> FunctionMetaData;
-};
+		TArray<UE::NNERuntimeIREE::FFunctionMetaData> FunctionMetaData;
+	};
+} // UE::NNERuntimeIREE
diff --git a/Engine/Plugins/Experimental/NNERuntimeIREE/Source/NNERuntimeIREE/Private/NNERuntimeIREE.cpp b/Engine/Plugins/Experimental/NNERuntimeIREE/Source/NNERuntimeIREE/Private/NNERuntimeIREE.cpp
--- a/Engine/Plugins/Experimental/NNERuntimeIREE/Source/NNERuntimeIREE/Private/NNERuntimeIREE.cpp
+++ b/Engine/Plugins/Experimental/NNERuntimeIREE/Source/NNERuntimeIREE/Private/NNERuntimeIREE.cpp
@@ -243,6 +243,7 @@ TSharedPtr<UE::NNE::FSharedModelData> UNNERuntimeIREECpu::CreateModelData(const
 	SCOPED_NAMED_EVENT_TEXT("UNNERuntimeIREECpu::CreateModelData", FColor::Magenta);
 
 #if WITH_EDITOR
+	using namespace UE::NNERuntimeIREE;
 	using namespace UE::NNERuntimeIREE::Private;
 
 	const FString TargetPlatformName = GetTargetPlatformName(TargetPlatform);
@@ -343,14 +344,14 @@ TSharedPtr<UE::NNE::FSharedModelData> UNNERuntimeIREECpu::CreateModelData(const
 		}
 		if (IREEModelData->ModuleMetaData.IsEmpty())
 		{
-			TWeakObjectPtr<UNNERuntimeIREEModuleMetaData> CompilerModuleMetaData = NewObject<UNNERuntimeIREEModuleMetaData>();
-			if (!CompilerModuleMetaData->ParseFromBuffer(FileData))
+			FModuleMetaData CompilerModuleMetaData{};
+			if (!CompilerModuleMetaData.ParseFromBuffer(FileData))
 			{
 				UE_LOG(LogNNERuntimeIREE, Warning, TEXT("UNNERuntimeIREECpu failed to parse meta data %s"), *FileIdString);
 				return TSharedPtr<UE::NNE::FSharedModelData>();
 			}
 			FMemoryWriter64 Writer(IREEModelData->ModuleMetaData, /*bIsPersitent =*/ true);
-			CompilerModuleMetaData->Serialize(Writer);
+			CompilerModuleMetaData.Serialize(Writer);
 		}
 		{
 			FMemoryWriter64 Writer(IREEModelData->CompilerResult, /*bIsPersitent =*/ true);
@@ -424,6 +425,7 @@ TSharedPtr<UE::NNE::IModelCPU> UNNERuntimeIREECpu::CreateModelCPU(const TObjectP
 {
 	check(ModelData != nullptr);
 
+	using namespace UE::NNERuntimeIREE;
 	using namespace UE::NNERuntimeIREE::Private;
 
 	if (CanCreateModelCPU(ModelData) != ECanCreateModelCPUStatus::Ok)
@@ -449,13 +451,13 @@ TSharedPtr<UE::NNE::IModelCPU> UNNERuntimeIREECpu::CreateModelCPU(const TObjectP
 		return TSharedPtr<UE::NNE::IModelCPU>();
 	}
 
-	TWeakObjectPtr<UNNERuntimeIREEModuleMetaData> ModuleMetaData = NewObject<UNNERuntimeIREEModuleMetaData>();
+	FModuleMetaData ModuleMetaData{};
 	{
 		FMemoryReaderView Reader(IREEModelData->ModuleMetaData, /*bIsPersitent =*/ true);
-		ModuleMetaData->Serialize(Reader);
+		ModuleMetaData.Serialize(Reader);
 	}
 
-	if (ModuleMetaData->FunctionMetaData.IsEmpty())
+	if (ModuleMetaData.FunctionMetaData.IsEmpty())
 	{
 		UE_LOG(LogNNERuntimeIREE, Warning, TEXT("UNNERuntimeIREECpu failed to parse the module meta data, please reimport the original model"));
 		return TSharedPtr<UE::NNE::IModelCPU>();
@@ -494,7 +496,7 @@ TSharedPtr<UE::NNE::IModelCPU> UNNERuntimeIREECpu::CreateModelCPU(const TObjectP
 	const FString SharedLibraryDirFullPath = FPaths::ConvertRelativePathToFull(FPaths::Combine(FPaths::ProjectDir(), GetSharedLibDirPath(UGameplayStatics::GetPlatformName(), GetRuntimeName(), FileIdString)));
 	const FString SharedLibrarySubDirFullPath = FPaths::Combine(SharedLibraryDirFullPath, ArchitectureInfo.RelativeDirPath);
 
-	TSharedPtr<UE::NNE::IModelCPU> Model = UE::NNERuntimeIREE::CPU::FModel::Make(*Environment, SharedLibrarySubDirFullPath, ArchitectureInfo.SharedLibraryFileName, ArchitectureInfo.VmfbFileName, ArchitectureInfo.SharedLibraryEntryPointName, *ModuleMetaData);
+	TSharedPtr<UE::NNE::IModelCPU> Model = UE::NNERuntimeIREE::CPU::FModel::Make(*Environment, SharedLibrarySubDirFullPath, ArchitectureInfo.SharedLibraryFileName, ArchitectureInfo.VmfbFileName, ArchitectureInfo.SharedLibraryEntryPointName, ModuleMetaData);
 	if (!Model.IsValid())
 	{
 		UE_LOG(LogNNERuntimeIREE, Warning, TEXT("UNNERuntimeIREECpu could not initialize the model created from model data with id %s"), *FileIdString);
@@ -671,6 +673,7 @@ TSharedPtr<UE::NNE::FSharedModelData> UNNERuntimeIREERdg::CreateModelData(const
 	SCOPED_NAMED_EVENT_TEXT("UNNERuntimeIREERdg::CreateModelData", FColor::Magenta);
 
 #if WITH_EDITOR
+	using namespace UE::NNERuntimeIREE;
 	using namespace UE::NNERuntimeIREE::Private;
 
 	const FString TargetPlatformName = GetTargetPlatformName(TargetPlatform);
@@ -772,16 +775,14 @@ TSharedPtr<UE::NNE::FSharedModelData> UNNERuntimeIREERdg::CreateModelData(const
 		}
 		if (IREEModelData->ModuleMetaData.IsEmpty())
 		{
-			TWeakObjectPtr<UNNERuntimeIREEModuleMetaData> CompilerModuleMetaData = NewObject<UNNERuntimeIREEModuleMetaData>();
-
-			if (!CompilerModuleMetaData->ParseFromBuffer(FileData))
+			FModuleMetaData CompilerModuleMetaData{};
+			if (!CompilerModuleMetaData.ParseFromBuffer(FileData))
 			{
 				UE_LOG(LogNNERuntimeIREE, Warning, TEXT("UNNERuntimeIREERdg failed to parse meta data %s"), *FileIdString);
 				return TSharedPtr<UE::NNE::FSharedModelData>();
 			}
-
 			FMemoryWriter64 Writer(IREEModelData->ModuleMetaData, /*bIsPersitent =*/ true);
-			CompilerModuleMetaData->Serialize(Writer);
+			CompilerModuleMetaData.Serialize(Writer);
 		}
 		{
 			FMemoryWriter64 Writer(IREEModelData->CompilerResult, /*bIsPersitent =*/ true);
@@ -851,6 +852,7 @@ TSharedPtr<UE::NNE::IModelRDG> UNNERuntimeIREERdg::CreateModelRDG(const TObjectP
 {
 	check(ModelData != nullptr);
 
+	using namespace UE::NNERuntimeIREE;
 	using namespace UE::NNERuntimeIREE::Private;
 
 	if (CanCreateModelRDG(ModelData) != ECanCreateModelRDGStatus::Ok)
@@ -876,13 +878,13 @@ TSharedPtr<UE::NNE::IModelRDG> UNNERuntimeIREERdg::CreateModelRDG(const TObjectP
 		return {};
 	}
 
-	TWeakObjectPtr<UNNERuntimeIREEModuleMetaData> ModuleMetaData = NewObject<UNNERuntimeIREEModuleMetaData>();
+	FModuleMetaData ModuleMetaData{};
 	{
 		FMemoryReaderView Reader(IREEModelData->ModuleMetaData, /*bIsPersitent =*/ true);
-		ModuleMetaData->Serialize(Reader);
+		ModuleMetaData.Serialize(Reader);
 	}
 
-	if (ModuleMetaData->FunctionMetaData.IsEmpty())
+	if (ModuleMetaData.FunctionMetaData.IsEmpty())
 	{
 		UE_LOG(LogNNERuntimeIREE, Warning, TEXT("UNNERuntimeIREERdg failed to parse the module meta data, please reimport the original model"));
 		return {};
@@ -925,7 +927,7 @@ TSharedPtr<UE::NNE::IModelRDG> UNNERuntimeIREERdg::CreateModelRDG(const TObjectP
 	const FString SharedLibraryDirFullPath = FPaths::ConvertRelativePathToFull(FPaths::Combine(FPaths::ProjectDir(), GetSharedLibDirPath(UGameplayStatics::GetPlatformName(), GetRuntimeName(), FileIdString), RuntimeSubdir));
 	const FString SharedLibrarySubDirFullPath = FPaths::Combine(SharedLibraryDirFullPath, BuildTargetResult->ShaderPlatform);
 
-	TSharedPtr<UE::NNE::IModelRDG> Model = UE::NNERuntimeIREE::RDG::FModel::Make(SharedLibrarySubDirFullPath, BuildTargetResult->VmfbData, *ModuleMetaData, ExecutableMap);
+	TSharedPtr<UE::NNE::IModelRDG> Model = UE::NNERuntimeIREE::RDG::FModel::Make(SharedLibrarySubDirFullPath, BuildTargetResult->VmfbData, ModuleMetaData, ExecutableMap);
 	if (!Model.IsValid())
 	{
 		UE_LOG(LogNNERuntimeIREE, Warning, TEXT("UNNERuntimeIREERdg could not initialize the model created from model data with id %s"), *FileIdString);
diff --git a/Engine/Plugins/Experimental/NNERuntimeIREE/Source/NNERuntimeIREE/Private/NNERuntimeIREEMetaData.cpp b/Engine/Plugins/Experimental/NNERuntimeIREE/Source/NNERuntimeIREE/Private/NNERuntimeIREEMetaData.cpp
--- a/Engine/Plugins/Experimental/NNERuntimeIREE/Source/NNERuntimeIREE/Private/NNERuntimeIREEMetaData.cpp
+++ b/Engine/Plugins/Experimental/NNERuntimeIREE/Source/NNERuntimeIREE/Private/NNERuntimeIREEMetaData.cpp
@@ -19,288 +19,291 @@
 
 #include "NNEMlirTools_cxx_api.h"
 
-namespace UE::NNERuntimeIREE::ModuleMetaData::Private
+namespace UE::NNERuntimeIREE
 {
-	enum Version : uint32
+	namespace ModuleMetaData::Private
 	{
-		V0 = 0, // Initial
-		// New versions can be added above this line
-		VersionPlusOne,
-		Latest = VersionPlusOne - 1
-	};
-	const FGuid GUID(0x2f9ffd31, 0x12b817cd, 0x627855bf, 0x5e405720);
-	FCustomVersionRegistration Version(GUID, Version::Latest, TEXT("NNERuntimeIREEModuleMetaDataVersion"));// Always save with the latest version
-
-	ENNETensorDataType ConvertTypeString(const FString& TypeString)
-	{
-		if (TypeString.StartsWith("char"))
-		{
-			return ENNETensorDataType::Char;
-		}
-		if (TypeString.StartsWith("bool") || TypeString.StartsWith("i1"))
-		{
-			return ENNETensorDataType::Boolean;
-		}
-		else if (TypeString.StartsWith("half"))
-		{
-			return ENNETensorDataType::Half;
-		}
-		else if (TypeString.StartsWith("bf16"))
+		enum Version : uint32
 		{
-			return ENNETensorDataType::BFloat16;
-		}
-		else if (TypeString.StartsWith("f"))
-		{
-			if (TypeString.StartsWith("f16"))
-			{
-				return ENNETensorDataType::Half;
-			}
-			else if (TypeString.StartsWith("float") || TypeString.StartsWith("f32"))
-			{
-				return ENNETensorDataType::Float;
-			}
-			else if (TypeString.StartsWith("f64"))
-			{
-				return ENNETensorDataType::Double;
-			}
-		}
-		else if (TypeString.StartsWith("double"))
-		{
-			return ENNETensorDataType::Double;
-		}
-		else if (TypeString.StartsWith("i") || TypeString.StartsWith("si"))
+			V0 = 0, // Initial
+			// New versions can be added above this line
+			VersionPlusOne,
+			Latest = VersionPlusOne - 1
+		};
+		const FGuid GUID(0x2f9ffd31, 0x12b817cd, 0x627855bf, 0x5e405720);
+		FCustomVersionRegistration Version(GUID, Version::Latest, TEXT("NNERuntimeIREEModuleMetaDataVersion"));// Always save with the latest version
+
+		ENNETensorDataType ConvertTypeString(const FString& TypeString)
 		{
-			if (TypeString.EndsWith("i8"))
+			if (TypeString.StartsWith("char"))
 			{
-				return ENNETensorDataType::Int8;
+				return ENNETensorDataType::Char;
 			}
-			else if (TypeString.EndsWith("i16"))
+			if (TypeString.StartsWith("bool") || TypeString.StartsWith("i1"))
 			{
-				return ENNETensorDataType::Int16;
+				return ENNETensorDataType::Boolean;
 			}
-			else if (TypeString.EndsWith("i32") || TypeString.EndsWith("int"))
+			else if (TypeString.StartsWith("half"))
 			{
-				return ENNETensorDataType::Int32;
+				return ENNETensorDataType::Half;
 			}
-			else if (TypeString.EndsWith("i64"))
+			else if (TypeString.StartsWith("bf16"))
 			{
-				return ENNETensorDataType::Int64;
+				return ENNETensorDataType::BFloat16;
 			}
-		}
-		else if (TypeString.StartsWith("ui"))
-		{
-			if (TypeString.EndsWith("i8"))
+			else if (TypeString.StartsWith("f"))
 			{
-				return ENNETensorDataType::UInt8;
+				if (TypeString.StartsWith("f16"))
+				{
+					return ENNETensorDataType::Half;
+				}
+				else if (TypeString.StartsWith("float") || TypeString.StartsWith("f32"))
+				{
+					return ENNETensorDataType::Float;
+				}
+				else if (TypeString.StartsWith("f64"))
+				{
+					return ENNETensorDataType::Double;
+				}
 			}
-			else if (TypeString.EndsWith("i16"))
+			else if (TypeString.StartsWith("double"))
 			{
-				return ENNETensorDataType::UInt16;
+				return ENNETensorDataType::Double;
 			}
-			else if (TypeString.EndsWith("i32"))
+			else if (TypeString.StartsWith("i") || TypeString.StartsWith("si"))
 			{
-				return ENNETensorDataType::UInt32;
+				if (TypeString.EndsWith("i8"))
+				{
+					return ENNETensorDataType::Int8;
+				}
+				else if (TypeString.EndsWith("i16"))
+				{
+					return ENNETensorDataType::Int16;
+				}
+				else if (TypeString.EndsWith("i32") || TypeString.EndsWith("int"))
+				{
+					return ENNETensorDataType::Int32;
+				}
+				else if (TypeString.EndsWith("i64"))
+				{
+					return ENNETensorDataType::Int64;
+				}
 			}
-			else if (TypeString.EndsWith("i64"))
+			else if (TypeString.StartsWith("ui"))
 			{
-				return ENNETensorDataType::UInt64;
+				if (TypeString.EndsWith("i8"))
+				{
+					return ENNETensorDataType::UInt8;
+				}
+				else if (TypeString.EndsWith("i16"))
+				{
+					return ENNETensorDataType::UInt16;
+				}
+				else if (TypeString.EndsWith("i32"))
+				{
+					return ENNETensorDataType::UInt32;
+				}
+				else if (TypeString.EndsWith("i64"))
+				{
+					return ENNETensorDataType::UInt64;
+				}
 			}
+			return ENNETensorDataType::None;
 		}
-		return ENNETensorDataType::None;
-	}
-} // UE::NNERuntimeIREE::ModuleMetaData::Private
-
-void UNNERuntimeIREEModuleMetaData::Serialize(FArchive& Ar)
-{
-	// Store the asset version (no effect in load)
-	Ar.UsingCustomVersion(UE::NNERuntimeIREE::ModuleMetaData::Private::GUID);
+	} // ModuleMetaData::Private
 
-	if (Ar.IsSaving() || Ar.IsCountingMemory()) 
+	void FModuleMetaData::Serialize(FArchive& Ar)
 	{
-		int32 NumItems = FunctionMetaData.Num();
-		Ar << NumItems;
-		for (int32 i = 0; i < NumItems; i++)
-		{
-			Ar << FunctionMetaData[i].Name;
-
-			int32 NumInputs = FunctionMetaData[i].InputDescs.Num();
-			Ar << NumInputs;
-			for (int32 j = 0; j < NumInputs; j++)
-			{
-				FString Name = FunctionMetaData[i].InputDescs[j].GetName();
-				Ar << Name;
-
-				ENNETensorDataType Type = FunctionMetaData[i].InputDescs[j].GetDataType();
-				Ar << Type;
-
-				TArray<int32> Shape = (TArray<int32>)FunctionMetaData[i].InputDescs[j].GetShape().GetData();
-				Ar << Shape;
-			}
+		// Store the asset version (no effect in load)
+		Ar.UsingCustomVersion(UE::NNERuntimeIREE::ModuleMetaData::Private::GUID);
 
-			int32 NumOutputs = FunctionMetaData[i].OutputDescs.Num();
-			Ar << NumOutputs;
-			for (int32 j = 0; j < NumOutputs; j++)
-			{
-				FString Name = FunctionMetaData[i].OutputDescs[j].GetName();
-				Ar << Name;
-
-				ENNETensorDataType Type = FunctionMetaData[i].OutputDescs[j].GetDataType();
-				Ar << Type;
-
-				TArray<int32> Shape = (TArray<int32>)FunctionMetaData[i].OutputDescs[j].GetShape().GetData();
-				Ar << Shape;
-			}
-		}
-	}
-	else
-	{
-		int32 NumItems = 0;
-		int32 NumInputs = 0;
-		int32 NumOutputs = 0;
-		UE::NNERuntimeIREE::FFunctionMetaData MetaData;
-		FString Name;
-		ENNETensorDataType Type;
-		TArray<int32> Shape;
-
-		switch (Ar.CustomVer(UE::NNERuntimeIREE::ModuleMetaData::Private::GUID))
+		if (Ar.IsSaving() || Ar.IsCountingMemory()) 
 		{
-		case UE::NNERuntimeIREE::ModuleMetaData::Private::Version::V0:
+			int32 NumItems = FunctionMetaData.Num();
 			Ar << NumItems;
-			FunctionMetaData.SetNum(NumItems, EAllowShrinking::Yes);
 			for (int32 i = 0; i < NumItems; i++)
 			{
-				Ar << MetaData.Name;
+				Ar << FunctionMetaData[i].Name;
 
+				int32 NumInputs = FunctionMetaData[i].InputDescs.Num();
 				Ar << NumInputs;
-				MetaData.InputDescs.Empty();
 				for (int32 j = 0; j < NumInputs; j++)
 				{
+					FString Name = FunctionMetaData[i].InputDescs[j].GetName();
 					Ar << Name;
+
+					ENNETensorDataType Type = FunctionMetaData[i].InputDescs[j].GetDataType();
 					Ar << Type;
+
+					TArray<int32> Shape = (TArray<int32>)FunctionMetaData[i].InputDescs[j].GetShape().GetData();
 					Ar << Shape;
-					MetaData.InputDescs.Add(UE::NNE::FTensorDesc::Make(Name, UE::NNE::FSymbolicTensorShape::Make(Shape), Type));
 				}
 
+				int32 NumOutputs = FunctionMetaData[i].OutputDescs.Num();
 				Ar << NumOutputs;
-				MetaData.OutputDescs.Empty();
 				for (int32 j = 0; j < NumOutputs; j++)
 				{
+					FString Name = FunctionMetaData[i].OutputDescs[j].GetName();
 					Ar << Name;
+
+					ENNETensorDataType Type = FunctionMetaData[i].OutputDescs[j].GetDataType();
 					Ar << Type;
+
+					TArray<int32> Shape = (TArray<int32>)FunctionMetaData[i].OutputDescs[j].GetShape().GetData();
 					Ar << Shape;
-					MetaData.OutputDescs.Add(UE::NNE::FTensorDesc::Make(Name, UE::NNE::FSymbolicTensorShape::Make(Shape), Type));
 				}
+			}
+		}
+		else
+		{
+			int32 NumItems = 0;
+			int32 NumInputs = 0;
+			int32 NumOutputs = 0;
+			UE::NNERuntimeIREE::FFunctionMetaData MetaData;
+			FString Name;
+			ENNETensorDataType Type;
+			TArray<int32> Shape;
 
-				FunctionMetaData[i] = MetaData;
+			switch (Ar.CustomVer(UE::NNERuntimeIREE::ModuleMetaData::Private::GUID))
+			{
+			case UE::NNERuntimeIREE::ModuleMetaData::Private::Version::V0:
+				Ar << NumItems;
+				FunctionMetaData.SetNum(NumItems, EAllowShrinking::Yes);
+				for (int32 i = 0; i < NumItems; i++)
+				{
+					Ar << MetaData.Name;
+
+					Ar << NumInputs;
+					MetaData.InputDescs.Empty();
+					for (int32 j = 0; j < NumInputs; j++)
+					{
+						Ar << Name;
+						Ar << Type;
+						Ar << Shape;
+						MetaData.InputDescs.Add(UE::NNE::FTensorDesc::Make(Name, UE::NNE::FSymbolicTensorShape::Make(Shape), Type));
+					}
+
+					Ar << NumOutputs;
+					MetaData.OutputDescs.Empty();
+					for (int32 j = 0; j < NumOutputs; j++)
+					{
+						Ar << Name;
+						Ar << Type;
+						Ar << Shape;
+						MetaData.OutputDescs.Add(UE::NNE::FTensorDesc::Make(Name, UE::NNE::FSymbolicTensorShape::Make(Shape), Type));
+					}
+
+					FunctionMetaData[i] = MetaData;
+				}
+				break;
+			default:
+				UE_LOG(LogNNERuntimeIREE, Error, TEXT("UNNERuntimeIREEModuleMetaData: Unknown asset version %d: Deserialisation failed, please reimport the original model."), Ar.CustomVer(UE::NNERuntimeIREE::ModuleMetaData::Private::GUID));
+				break;
 			}
-			break;
-		default:
-			UE_LOG(LogNNERuntimeIREE, Error, TEXT("UNNERuntimeIREEModuleMetaData: Unknown asset version %d: Deserialisation failed, please reimport the original model."), Ar.CustomVer(UE::NNERuntimeIREE::ModuleMetaData::Private::GUID));
-			break;
 		}
 	}
-}
 
-bool UNNERuntimeIREEModuleMetaData::ParseFromBuffer(TConstArrayView64<uint8> Buffer)
-{
-	using namespace UE::NNE;
-	using namespace UE::NNERuntimeIREE::ModuleMetaData::Private;
+	bool FModuleMetaData::ParseFromBuffer(TConstArrayView64<uint8> Buffer)
+	{
+		using namespace UE::NNE;
+		using namespace UE::NNERuntimeIREE::ModuleMetaData::Private;
 
-	const char* BufferPtr = reinterpret_cast<const char*>(Buffer.GetData());
-	const size_t BufferSize = static_cast<size_t>(Buffer.Num());
+		const char* BufferPtr = reinterpret_cast<const char*>(Buffer.GetData());
+		const size_t BufferSize = static_cast<size_t>(Buffer.Num());
 
 #if WITH_EDITOR
-	try
+		try
 #endif // WITH_EDITOR
-	{
-		NNEMlirTools::Context Ctx;
-		const NNEMlirTools::Module Mod = NNEMlirTools::Module::ParseFromBuffer(Ctx, BufferPtr, BufferSize);
-
-		auto MakeValueName = [] (const FString& Name, const FString& DefaultPrefix, int32 DefaultIndex)
 		{
-			return !Name.IsEmpty() ? Name : (DefaultPrefix + FString::FromInt(DefaultIndex));
-		};
-
-		auto GetShape = [] (const NNEMlirTools::Value& Val)
-		{
-			const int32 Rank = static_cast<int32>(Val.GetRank());
-
-			TArray<int64_t> Shape64;
-			Shape64.SetNumUninitialized(Rank);
-			Val.GetShape(Shape64.GetData(), Rank);
+			NNEMlirTools::Context Ctx;
+			const NNEMlirTools::Module Mod = NNEMlirTools::Module::ParseFromBuffer(Ctx, BufferPtr, BufferSize);
 
-			TArray<int32> Result;
-			Result.SetNumUninitialized(Rank);
+			auto MakeValueName = [] (const FString& Name, const FString& DefaultPrefix, int32 DefaultIndex)
+			{
+				return !Name.IsEmpty() ? Name : (DefaultPrefix + FString::FromInt(DefaultIndex));
+			};
 
-			for (int i = 0; i < Rank; i++)
+			auto GetShape = [] (const NNEMlirTools::Value& Val)
 			{
-				Result[i] = static_cast<int32>(Shape64[i]);
-			}
+				const int32 Rank = static_cast<int32>(Val.GetRank());
 
-			return Result;
-		};
+				TArray<int64_t> Shape64;
+				Shape64.SetNumUninitialized(Rank);
+				Val.GetShape(Shape64.GetData(), Rank);
 
-		auto GetTensorDesc = [&GetShape] (const FString& Name, const NNEMlirTools::Value& Val)
-		{
-			const TArray<int32> Shape = GetShape(Val);
-			FSymbolicTensorShape SymbolicShape = FSymbolicTensorShape::Make(Shape);
+				TArray<int32> Result;
+				Result.SetNumUninitialized(Rank);
 
-			const FString TypeStr = Val.GetElementType().c_str();
-			check(!TypeStr.IsEmpty());
+				for (int i = 0; i < Rank; i++)
+				{
+					Result[i] = static_cast<int32>(Shape64[i]);
+				}
 
-			const ENNETensorDataType DataType = ConvertTypeString(TypeStr);
+				return Result;
+			};
 
-			return FTensorDesc::Make(Name, SymbolicShape, DataType);
-		};
+			auto GetTensorDesc = [&GetShape] (const FString& Name, const NNEMlirTools::Value& Val)
+			{
+				const TArray<int32> Shape = GetShape(Val);
+				FSymbolicTensorShape SymbolicShape = FSymbolicTensorShape::Make(Shape);
 
-		const int32 NumFun = Mod.GetFunctionCount();
-		FunctionMetaData.Reset(NumFun);
+				const FString TypeStr = Val.GetElementType().c_str();
+				check(!TypeStr.IsEmpty());
 
-		for (int32 i = 0; i < NumFun; i++)
-		{
-			const NNEMlirTools::Function Fun = Mod.GetFunction(i);
+				const ENNETensorDataType DataType = ConvertTypeString(TypeStr);
 
-			UE::NNERuntimeIREE::FFunctionMetaData MetaData;
-			MetaData.Name = Fun.GetName().c_str();
-			check(!MetaData.Name.IsEmpty());
+				return FTensorDesc::Make(Name, SymbolicShape, DataType);
+			};
 
-			const int32 NumInputs = Fun.GetInputCount();
-			MetaData.InputDescs.Reserve(NumInputs);
+			const int32 NumFun = Mod.GetFunctionCount();
+			FunctionMetaData.Reset(NumFun);
 
-			for (int j = 0; j < NumInputs; j++)
+			for (int32 i = 0; i < NumFun; i++)
 			{
-				const NNEMlirTools::Value Input = Fun.GetInput(j);
-				const FString Name = MakeValueName(Input.GetName().c_str(), TEXT("arg"), j);
+				const NNEMlirTools::Function Fun = Mod.GetFunction(i);
 
-				MetaData.InputDescs.Add(GetTensorDesc(Name, Input));
-			}
+				UE::NNERuntimeIREE::FFunctionMetaData MetaData;
+				MetaData.Name = Fun.GetName().c_str();
+				check(!MetaData.Name.IsEmpty());
+
+				const int32 NumInputs = Fun.GetInputCount();
+				MetaData.InputDescs.Reserve(NumInputs);
+
+				for (int j = 0; j < NumInputs; j++)
+				{
+					const NNEMlirTools::Value Input = Fun.GetInput(j);
+					const FString Name = MakeValueName(Input.GetName().c_str(), TEXT("arg"), j);
 
-			const int32 NumOutputs = Fun.GetResultCount();
-			MetaData.OutputDescs.Reserve(NumOutputs);
+					MetaData.InputDescs.Add(GetTensorDesc(Name, Input));
+				}
+
+				const int32 NumOutputs = Fun.GetResultCount();
+				MetaData.OutputDescs.Reserve(NumOutputs);
 			
-			for (int j = 0; j < NumOutputs; j++)
-			{
-				const NNEMlirTools::Value Output = Fun.GetResult(j);
-				const FString Name = MakeValueName(Output.GetName().c_str(), TEXT("res"), j);
+				for (int j = 0; j < NumOutputs; j++)
+				{
+					const NNEMlirTools::Value Output = Fun.GetResult(j);
+					const FString Name = MakeValueName(Output.GetName().c_str(), TEXT("res"), j);
 
-				MetaData.OutputDescs.Add(GetTensorDesc(Name, Output));
-			}
+					MetaData.OutputDescs.Add(GetTensorDesc(Name, Output));
+				}
 
-			FunctionMetaData.Add(MoveTemp(MetaData));
+				FunctionMetaData.Add(MoveTemp(MetaData));
+			}
 		}
-	}
 #if WITH_EDITOR
-	catch (const std::exception& Exception)
-	{
-		UE_LOG(LogNNERuntimeIREE, Error, TEXT("Parse MLIR exception: %hs"), Exception.what());
-		return false;
-	}
-	catch (...)
-	{
-		UE_LOG(LogNNERuntimeIREE, Error, TEXT("Parse MLIR unknown exception."));
-		return false;
-	}
+		catch (const std::exception& Exception)
+		{
+			UE_LOG(LogNNERuntimeIREE, Error, TEXT("Parse MLIR exception: %hs"), Exception.what());
+			return false;
+		}
+		catch (...)
+		{
+			UE_LOG(LogNNERuntimeIREE, Error, TEXT("Parse MLIR unknown exception."));
+			return false;
+		}
 #endif // WITH_EDITOR
 
-	return true;
-}
\ No newline at end of file
+		return true;
+	}
+} // UE::NNERuntimeIREE
diff --git a/Engine/Plugins/Experimental/NNERuntimeIREE/Source/NNERuntimeIREE/Private/NNERuntimeIREEModel.cpp b/Engine/Plugins/Experimental/NNERuntimeIREE/Source/NNERuntimeIREE/Private/NNERuntimeIREEModel.cpp
--- a/Engine/Plugins/Experimental/NNERuntimeIREE/Source/NNERuntimeIREE/Private/NNERuntimeIREEModel.cpp
+++ b/Engine/Plugins/Experimental/NNERuntimeIREE/Source/NNERuntimeIREE/Private/NNERuntimeIREEModel.cpp
@@ -442,7 +442,7 @@ namespace UE::NNERuntimeIREE
 				return MakeShared<UE::NNE::FSharedModelData>(FSharedBuffer::TakeOwnership(Data, DataSize, FMemory::Free), IREE_HAL_HEAP_BUFFER_ALIGNMENT);
 			}
 
-			static TSharedPtr<FModule> Make(const FString& InVmfbPath, TSharedRef<UE::NNE::FSharedModelData> ModelData, const UNNERuntimeIREEModuleMetaData& InModuleMetaData)
+			static TSharedPtr<FModule> Make(const FString& InVmfbPath, TSharedRef<UE::NNE::FSharedModelData> ModelData, const FModuleMetaData& InModuleMetaData)
 			{
 				check(!InVmfbPath.IsEmpty());
 				check(!InModuleMetaData.FunctionMetaData.IsEmpty());
@@ -1072,7 +1072,7 @@ namespace UE::NNERuntimeIREE
 
 		}
 
-		TSharedPtr<FModel> FModel::Make(UE::NNERuntimeIREE::Private::FEnvironment& InEnvironment, const FString& InDirPath, const FString& InSharedLibraryFileName, const FString& InVmfbFileName, const FString& InLibraryQueryFunctionName, const UNNERuntimeIREEModuleMetaData& InModuleMetaData)
+		TSharedPtr<FModel> FModel::Make(UE::NNERuntimeIREE::Private::FEnvironment& InEnvironment, const FString& InDirPath, const FString& InSharedLibraryFileName, const FString& InVmfbFileName, const FString& InLibraryQueryFunctionName, const FModuleMetaData& InModuleMetaData)
 		{
 			check(!InSharedLibraryFileName.IsEmpty());
 			check(!InVmfbFileName.IsEmpty());
@@ -1613,7 +1613,7 @@ namespace UE::NNERuntimeIREE
 
 		}
 
-		TSharedPtr<FModel> FModel::Make(const FString& InDirPath, TConstArrayView64<uint8> VmfbData, const UNNERuntimeIREEModuleMetaData& InModuleMetaData, const TMap<FString, TConstArrayView<uint8>>& Executables)
+		TSharedPtr<FModel> FModel::Make(const FString& InDirPath, TConstArrayView64<uint8> VmfbData, const FModuleMetaData& InModuleMetaData, const TMap<FString, TConstArrayView<uint8>>& Executables)
 		{
 			check(!InDirPath.IsEmpty());
 			check(!VmfbData.IsEmpty());
diff --git a/Engine/Plugins/Experimental/NNERuntimeIREE/Source/NNERuntimeIREE/Private/NNERuntimeIREEModel.h b/Engine/Plugins/Experimental/NNERuntimeIREE/Source/NNERuntimeIREE/Private/NNERuntimeIREEModel.h
--- a/Engine/Plugins/Experimental/NNERuntimeIREE/Source/NNERuntimeIREE/Private/NNERuntimeIREEModel.h
+++ b/Engine/Plugins/Experimental/NNERuntimeIREE/Source/NNERuntimeIREE/Private/NNERuntimeIREEModel.h
@@ -59,7 +59,7 @@ namespace UE::NNERuntimeIREE
 			FModel(TSharedRef<UE::NNERuntimeIREE::CPU::Private::FDevice> InDevice, TSharedRef<UE::NNERuntimeIREE::Private::FModule> InModule);
 
 		public:
-			static TSharedPtr<FModel> Make(UE::NNERuntimeIREE::Private::FEnvironment& InEnvironment, const FString& InDirPath, const FString& InSharedLibraryFileName, const FString& InVmfbFileName, const FString& InLibraryQueryFunctionName, const UNNERuntimeIREEModuleMetaData& InModuleMetaData);
+			static TSharedPtr<FModel> Make(UE::NNERuntimeIREE::Private::FEnvironment& InEnvironment, const FString& InDirPath, const FString& InSharedLibraryFileName, const FString& InVmfbFileName, const FString& InLibraryQueryFunctionName, const FModuleMetaData& InModuleMetaData);
 
 		public:
 			//~ Begin IModelCPU Interface
@@ -112,7 +112,7 @@ namespace UE::NNERuntimeIREE
 			FModel(TSharedRef<UE::NNERuntimeIREE::RDG::Private::FDevice> InDevice, TSharedRef<UE::NNERuntimeIREE::Private::FModule> InModule);
 
 		public:
-			static TSharedPtr<FModel> Make(const FString& InDirPath, TConstArrayView64<uint8> VmfbData, const UNNERuntimeIREEModuleMetaData& InModuleMetaData, const TMap<FString, TConstArrayView<uint8>>& Executables);
+			static TSharedPtr<FModel> Make(const FString& InDirPath, TConstArrayView64<uint8> VmfbData, const FModuleMetaData& InModuleMetaData, const TMap<FString, TConstArrayView<uint8>>& Executables);
 
 		public:
 			//~ Begin IModelRDG Interface
diff --git a/Engine/Plugins/Experimental/NNERuntimeIREE/Source/NNERuntimeIREEEditor/Private/NNERuntimeIREEModelDataFactory.cpp b/Engine/Plugins/Experimental/NNERuntimeIREE/Source/NNERuntimeIREEEditor/Private/NNERuntimeIREEModelDataFactory.cpp
--- a/Engine/Plugins/Experimental/NNERuntimeIREE/Source/NNERuntimeIREEEditor/Private/NNERuntimeIREEModelDataFactory.cpp
+++ b/Engine/Plugins/Experimental/NNERuntimeIREE/Source/NNERuntimeIREEEditor/Private/NNERuntimeIREEModelDataFactory.cpp
@@ -158,19 +158,19 @@ UObject* UNNERuntimeIREEModelDataFactory::FactoryCreateBinary(UClass* Class, UOb
 
 	TConstArrayView<uint8> BufferView = MakeArrayView(Buffer, BufferEnd - Buffer);
 
-	TObjectPtr<UNNERuntimeIREEModuleMetaData> ModuleMetaData = NewObject<UNNERuntimeIREEModuleMetaData>();
-	if (!ModuleMetaData->ParseFromBuffer(BufferView))
+	UE::NNERuntimeIREE::FModuleMetaData ModuleMetaData{};
+	if (!ModuleMetaData.ParseFromBuffer(BufferView))
 	{
 		UE_LOG(LogNNERuntimeIREE, Error, TEXT("UNNERuntimeIREEModelDataFactory failed to parse the models meta data"));
 		return nullptr;
 	}
-	if (ModuleMetaData->FunctionMetaData.IsEmpty())
+	if (ModuleMetaData.FunctionMetaData.IsEmpty())
 	{
 		UE_LOG(LogNNERuntimeIREE, Error, TEXT("UNNERuntimeIREEModelDataFactory no public functions found in models meta data"));
 		return nullptr;
 	}
 
-	if (ModuleMetaData->FunctionMetaData.Num() > 1)
+	if (ModuleMetaData.FunctionMetaData.Num() > 1)
 	{
 		TSharedPtr<SWindow> ParentWindow;
 		if (FModuleManager::Get().IsModuleLoaded("MainFrame"))
@@ -188,7 +188,7 @@ UObject* UNNERuntimeIREEModelDataFactory::FactoryCreateBinary(UClass* Class, UOb
 		Window->SetContent
 		(
 			SAssignNew(ImportWindow, SMLIRImportWindow)
-			.MetaData(ModuleMetaData->FunctionMetaData)
+			.MetaData(ModuleMetaData.FunctionMetaData)
 			.WidgetWindow(Window)
 		);
 
@@ -200,14 +200,14 @@ UObject* UNNERuntimeIREEModelDataFactory::FactoryCreateBinary(UClass* Class, UOb
 			return nullptr;
 		}
 
-		UE::NNERuntimeIREE::FFunctionMetaData Swapper = ModuleMetaData->FunctionMetaData[0];
-		ModuleMetaData->FunctionMetaData[0] = ModuleMetaData->FunctionMetaData[ImportWindow->GetCurrentSelectedFunctionIndex()];
-		ModuleMetaData->FunctionMetaData[ImportWindow->GetCurrentSelectedFunctionIndex()] = Swapper;
+		UE::NNERuntimeIREE::FFunctionMetaData Swapper = ModuleMetaData.FunctionMetaData[0];
+		ModuleMetaData.FunctionMetaData[0] = ModuleMetaData.FunctionMetaData[ImportWindow->GetCurrentSelectedFunctionIndex()];
+		ModuleMetaData.FunctionMetaData[ImportWindow->GetCurrentSelectedFunctionIndex()] = Swapper;
 	}
 
 	TArray<uint8> MetaDataByteArray;
 	FMemoryWriter Writer(MetaDataByteArray, /*bIsPersitent =*/ true);
-	ModuleMetaData->Serialize(Writer);
+	ModuleMetaData.Serialize(Writer);
 
 	TMap<FString, TConstArrayView64<uint8>> AdditionalFileData;
 	AdditionalFileData.Add("IREEModuleMetaData", MetaDataByteArray);
