[NNE] NNERuntimeIREE replace TArray<uint8> inside ModelData with the actual types

#rnx
#tests NNEModelTest including tests that are imported with an earlier version

This allows us to load the entiere ModelData in one Serialize call, which simplifies the code and reduces the number of error cases.

#Synced-CL 49758318

diff --git a/Engine/Plugins/Experimental/NNERuntimeIREE/Source/NNERuntimeIREE/Private/NNERuntimeIREE.cpp b/Engine/Plugins/Experimental/NNERuntimeIREE/Source/NNERuntimeIREE/Private/NNERuntimeIREE.cpp
--- a/Engine/Plugins/Experimental/NNERuntimeIREE/Source/NNERuntimeIREE/Private/NNERuntimeIREE.cpp
+++ b/Engine/Plugins/Experimental/NNERuntimeIREE/Source/NNERuntimeIREE/Private/NNERuntimeIREE.cpp
@@ -210,7 +210,7 @@ namespace UE::NNERuntimeIREE::Private
 } // UE::NNERuntimeIREE::Private
 
 FGuid UNNERuntimeIREECpu::GUID = FGuid((int32)'I', (int32)'C', (int32)'P', (int32)'U');
-int32 UNNERuntimeIREECpu::Version = 0x00000005;
+int32 UNNERuntimeIREECpu::Version = 0x00000006; // Increment each time iree compiler/runtime is updated and each time FModelData serialization changes
 
 void UNNERuntimeIREECpu::Init(TSharedRef<UE::NNERuntimeIREE::Private::FEnvironment> InEnvironment)
 {
@@ -265,7 +265,6 @@ TSharedPtr<UE::NNE::FSharedModelData> UNNERuntimeIREECpu::CreateModelData(const
 
 	TArray64<uint8> ResultData;
 	TWeakObjectPtr<UNNERuntimeIREEModelDataCPU> IREEModelData = NewObject<UNNERuntimeIREEModelDataCPU>();
-	FNNERuntimeIREECompilerResultCPU CompilerResult{};
 
 	bool bNeedCompileMlir = true;
 	if (PlatformFile.FileExists(*IREEModelDataFilePath) && FFileHelper::LoadFileToArray(ResultData, *IREEModelDataFilePath))
@@ -281,15 +280,10 @@ TSharedPtr<UE::NNE::FSharedModelData> UNNERuntimeIREECpu::CreateModelData(const
 
 			check(FileIdString.Equals(IREEModelData->FileId.ToString(EGuidFormats::Digits).ToLower()));
 
-			{
-				FMemoryReaderView Reader(IREEModelData->CompilerResult, /*bIsPersitent =*/ true);
-				FNNERuntimeIREECompilerResultCPU::StaticStruct()->SerializeBin(Reader, &CompilerResult);
-			}
-
 			bNeedCompileMlir = false;
-			for (int32 i = 0; i < CompilerResult.ArchitectureInfos.Num(); i++)
+			for (int32 i = 0; i < IREEModelData->CompilerResult.ArchitectureInfos.Num(); i++)
 			{
-				const FNNERuntimeIREEArchitectureInfoCPU& Info = CompilerResult.ArchitectureInfos[i];
+				const FNNERuntimeIREEArchitectureInfoCPU& Info = IREEModelData->CompilerResult.ArchitectureInfos[i];
 				const FString SharedLibrarySubDirFullPath = FPaths::Combine(SharedLibraryDirFullPath, Info.RelativeDirPath);
 
 				const FString SharedLibraryFilePath = FPaths::Combine(SharedLibrarySubDirFullPath, Info.SharedLibraryFileName);
@@ -327,9 +321,8 @@ TSharedPtr<UE::NNE::FSharedModelData> UNNERuntimeIREECpu::CreateModelData(const
 		}
 
 		// NOTE: From this point FileData refers to mlir data in any case
-		
 
-		if (!Compiler->CompileMlir(FileData, FileIdString, IntermediateDirFullPath, CompilerResult))
+		if (!Compiler->CompileMlir(FileData, FileIdString, IntermediateDirFullPath, IREEModelData->CompilerResult))
 		{
 			UE_LOG(LogNNERuntimeIREE, Warning, TEXT("UNNERuntimeIREECpu failed to compile model %s"), *FileIdString);
 			return TSharedPtr<UE::NNE::FSharedModelData>();
@@ -338,24 +331,18 @@ TSharedPtr<UE::NNE::FSharedModelData> UNNERuntimeIREECpu::CreateModelData(const
 		IREEModelData->GUID = UNNERuntimeIREECpu::GUID;
 		IREEModelData->Version = UNNERuntimeIREECpu::Version;
 		IREEModelData->FileId = FileId;
-		if (AdditionalFileData.Contains("IREEModuleMetaData"))
+		if (AdditionalFileData.Contains("IREEModuleMetaData") && !AdditionalFileData["IREEModuleMetaData"].IsEmpty())
 		{
-			IREEModelData->ModuleMetaData = AdditionalFileData["IREEModuleMetaData"];
+			FMemoryReaderView Reader(AdditionalFileData["IREEModuleMetaData"], /*bIsPersitent =*/ true);
+			IREEModelData->ModuleMetaData.Serialize(Reader);
 		}
-		if (IREEModelData->ModuleMetaData.IsEmpty())
+		else
 		{
-			FModuleMetaData CompilerModuleMetaData{};
-			if (!CompilerModuleMetaData.ParseFromBuffer(FileData))
+			if (!IREEModelData->ModuleMetaData.ParseFromBuffer(FileData))
 			{
 				UE_LOG(LogNNERuntimeIREE, Warning, TEXT("UNNERuntimeIREECpu failed to parse meta data %s"), *FileIdString);
 				return TSharedPtr<UE::NNE::FSharedModelData>();
 			}
-			FMemoryWriter64 Writer(IREEModelData->ModuleMetaData, /*bIsPersitent =*/ true);
-			CompilerModuleMetaData.Serialize(Writer);
-		}
-		{
-			FMemoryWriter64 Writer(IREEModelData->CompilerResult, /*bIsPersitent =*/ true);
-			FNNERuntimeIREECompilerResultCPU::StaticStruct()->SerializeBin(Writer, &CompilerResult);
 		}
 
 		{
@@ -371,11 +358,11 @@ TSharedPtr<UE::NNE::FSharedModelData> UNNERuntimeIREECpu::CreateModelData(const
 	{
 		// Copy files for staging
 		FString StagingDirFullPath = FPaths::ConvertRelativePathToFull(FPaths::Combine(FPaths::ProjectDir(), GetPackagedModelDirPath(TargetPlatformName, GetRuntimeName(), FileIdString)));
-		for (int32 i = 0; i < CompilerResult.ArchitectureInfos.Num(); i++)
+		for (int32 i = 0; i < IREEModelData->CompilerResult.ArchitectureInfos.Num(); i++)
 		{
 			SCOPED_NAMED_EVENT_TEXT("Copy", FColor::Magenta);
 
-			const FNNERuntimeIREEArchitectureInfoCPU& Info = CompilerResult.ArchitectureInfos[i];
+			const FNNERuntimeIREEArchitectureInfoCPU& Info = IREEModelData->CompilerResult.ArchitectureInfos[i];
 			const FString SharedLibrarySubDirFullPath = FPaths::Combine(SharedLibraryDirFullPath, Info.RelativeDirPath);
 			const FString StagingSubDirFullPath = FPaths::Combine(StagingDirFullPath, Info.RelativeDirPath);
 			
@@ -445,36 +432,12 @@ TSharedPtr<UE::NNE::IModelCPU> UNNERuntimeIREECpu::CreateModelCPU(const TObjectP
 		IREEModelData->Serialize(Reader);
 	}
 
-	if (IREEModelData->ModuleMetaData.IsEmpty())
-	{
-		UE_LOG(LogNNERuntimeIREE, Warning, TEXT("UNNERuntimeIREECpu failed to find any module meta data, please reimport the original model"));
-		return TSharedPtr<UE::NNE::IModelCPU>();
-	}
-
-	FModuleMetaData ModuleMetaData{};
-	{
-		FMemoryReaderView Reader(IREEModelData->ModuleMetaData, /*bIsPersitent =*/ true);
-		ModuleMetaData.Serialize(Reader);
-	}
-
-	if (ModuleMetaData.FunctionMetaData.IsEmpty())
-	{
-		UE_LOG(LogNNERuntimeIREE, Warning, TEXT("UNNERuntimeIREECpu failed to parse the module meta data, please reimport the original model"));
-		return TSharedPtr<UE::NNE::IModelCPU>();
-	}
-
-	FNNERuntimeIREECompilerResultCPU CompilerResult{};
-	{
-		FMemoryReaderView Reader(IREEModelData->CompilerResult, /*bIsPersitent =*/ true);
-		FNNERuntimeIREECompilerResultCPU::StaticStruct()->SerializeBin(Reader, &CompilerResult);
-	}
-
 	int32 TargetIndex = -1;
-	for (int32 i = 0; i < CompilerResult.ArchitectureInfos.Num(); i++)
+	for (int32 i = 0; i < IREEModelData->CompilerResult.ArchitectureInfos.Num(); i++)
 	{
 		bool bMatches = true;
-		bMatches &= DoesArchitectureMatch(CompilerResult.ArchitectureInfos[i].Architecture);
-		bMatches &= DoX86FeaturesMatch(CompilerResult.ArchitectureInfos[i].X86Features);
+		bMatches &= DoesArchitectureMatch(IREEModelData->CompilerResult.ArchitectureInfos[i].Architecture);
+		bMatches &= DoX86FeaturesMatch(IREEModelData->CompilerResult.ArchitectureInfos[i].X86Features);
 
 		if (bMatches)
 		{
@@ -482,21 +445,18 @@ TSharedPtr<UE::NNE::IModelCPU> UNNERuntimeIREECpu::CreateModelCPU(const TObjectP
 			break;
 		}
 	}
-
 	if (TargetIndex < 0)
 	{
 		UE_LOG(LogNNERuntimeIREE, Warning, TEXT("UNNERuntimeIREECpu failed to find a matching target for the current hardware"));
 		return TSharedPtr<UE::NNE::IModelCPU>();
 	}
 
-
-	const FNNERuntimeIREEArchitectureInfoCPU& ArchitectureInfo = CompilerResult.ArchitectureInfos[TargetIndex];
-
+	const FNNERuntimeIREEArchitectureInfoCPU& ArchitectureInfo = IREEModelData->CompilerResult.ArchitectureInfos[TargetIndex];
 	const FString FileIdString = IREEModelData->FileId.ToString(EGuidFormats::Digits).ToLower();
 	const FString SharedLibraryDirFullPath = FPaths::ConvertRelativePathToFull(FPaths::Combine(FPaths::ProjectDir(), GetSharedLibDirPath(UGameplayStatics::GetPlatformName(), GetRuntimeName(), FileIdString)));
 	const FString SharedLibrarySubDirFullPath = FPaths::Combine(SharedLibraryDirFullPath, ArchitectureInfo.RelativeDirPath);
 
-	TSharedPtr<UE::NNE::IModelCPU> Model = UE::NNERuntimeIREE::CPU::FModel::Make(*Environment, SharedLibrarySubDirFullPath, ArchitectureInfo.SharedLibraryFileName, ArchitectureInfo.VmfbFileName, ArchitectureInfo.SharedLibraryEntryPointName, ModuleMetaData);
+	TSharedPtr<UE::NNE::IModelCPU> Model = UE::NNERuntimeIREE::CPU::FModel::Make(*Environment, SharedLibrarySubDirFullPath, ArchitectureInfo.SharedLibraryFileName, ArchitectureInfo.VmfbFileName, ArchitectureInfo.SharedLibraryEntryPointName, IREEModelData->ModuleMetaData);
 	if (!Model.IsValid())
 	{
 		UE_LOG(LogNNERuntimeIREE, Warning, TEXT("UNNERuntimeIREECpu could not initialize the model created from model data with id %s"), *FileIdString);
@@ -646,7 +606,7 @@ int32 UNNERuntimeIREEVulkan::GetVersion() const
 #ifdef WITH_NNE_RUNTIME_IREE_RDG
 
 FGuid UNNERuntimeIREERdg::GUID = FGuid((int32)'I', (int32)'R', (int32)'D', (int32)'G');
-int32 UNNERuntimeIREERdg::Version = 0x00000002;
+int32 UNNERuntimeIREERdg::Version = 0x00000003; // Increment each time iree compiler/runtime is updated and each time FModelData serialization changes
 
 FString UNNERuntimeIREERdg::GetRuntimeName() const
 {
@@ -731,8 +691,6 @@ TSharedPtr<UE::NNE::FSharedModelData> UNNERuntimeIREERdg::CreateModelData(const
 		PlatformFile.DeleteDirectoryRecursively(*IntermediateDirFullPath);
 		PlatformFile.CreateDirectoryTree(*IntermediateDirFullPath);
 
-		FIREECompilerRDGResult CompilerResult{};
-
 		TUniquePtr<UE::IREE::Compiler::RDG::FCompiler> Compiler = UE::IREE::Compiler::RDG::FCompiler::Make(TargetPlatform);
 		if (!Compiler.IsValid())
 		{
@@ -753,13 +711,13 @@ TSharedPtr<UE::NNE::FSharedModelData> UNNERuntimeIREERdg::CreateModelData(const
 
 		// NOTE: From this point FileData refers to mlir data in any case
 
-		if (!Compiler->CompileMlir(FileData, FileIdString, IntermediateDirFullPath, ShaderPlatforms, CompilerResult))
+		TWeakObjectPtr<UNNERuntimeIREEModelDataRDG> IREEModelData = NewObject<UNNERuntimeIREEModelDataRDG>();
+		if (!Compiler->CompileMlir(FileData, FileIdString, IntermediateDirFullPath, ShaderPlatforms, IREEModelData->CompilerResult))
 		{
 			UE_LOG(LogNNERuntimeIREE, Warning, TEXT("UNNERuntimeIREERdg failed to compile model %s"), *FileIdString);
 			return TSharedPtr<UE::NNE::FSharedModelData>();
 		}
 
-		TWeakObjectPtr<UNNERuntimeIREEModelDataRDG> IREEModelData = NewObject<UNNERuntimeIREEModelDataRDG>();
 		IREEModelData->Header.GUID = UNNERuntimeIREERdg::GUID;
 		IREEModelData->Header.Version = UNNERuntimeIREERdg::Version;
 		IREEModelData->Header.FileId = FileId;
@@ -769,24 +727,18 @@ TSharedPtr<UE::NNE::FSharedModelData> UNNERuntimeIREERdg::CreateModelData(const
 
 			IREEModelData->Header.ShaderPlatforms.Add(ShaderPlatformName);
 		}
-		if (AdditionalFileData.Contains("IREEModuleMetaData"))
+		if (AdditionalFileData.Contains("IREEModuleMetaData") && !AdditionalFileData["IREEModuleMetaData"].IsEmpty())
 		{
-			IREEModelData->ModuleMetaData = AdditionalFileData["IREEModuleMetaData"];
+			FMemoryReaderView Reader(AdditionalFileData["IREEModuleMetaData"]);
+			IREEModelData->ModuleMetaData.Serialize(Reader);
 		}
-		if (IREEModelData->ModuleMetaData.IsEmpty())
+		else
 		{
-			FModuleMetaData CompilerModuleMetaData{};
-			if (!CompilerModuleMetaData.ParseFromBuffer(FileData))
+			if (!IREEModelData->ModuleMetaData.ParseFromBuffer(FileData))
 			{
 				UE_LOG(LogNNERuntimeIREE, Warning, TEXT("UNNERuntimeIREERdg failed to parse meta data %s"), *FileIdString);
 				return TSharedPtr<UE::NNE::FSharedModelData>();
 			}
-			FMemoryWriter64 Writer(IREEModelData->ModuleMetaData, /*bIsPersitent =*/ true);
-			CompilerModuleMetaData.Serialize(Writer);
-		}
-		{
-			FMemoryWriter64 Writer(IREEModelData->CompilerResult, /*bIsPersitent =*/ true);
-			FIREECompilerRDGResult::StaticStruct()->SerializeBin(Writer, &CompilerResult);
 		}
 
 		{
@@ -871,33 +823,9 @@ TSharedPtr<UE::NNE::IModelRDG> UNNERuntimeIREERdg::CreateModelRDG(const TObjectP
 		FMemoryReaderView Reader(SharedDataView, /*bIsPersitent =*/ true);
 		IREEModelData->Serialize(Reader);
 	}
-
-	if (IREEModelData->ModuleMetaData.IsEmpty())
-	{
-		UE_LOG(LogNNERuntimeIREE, Warning, TEXT("UNNERuntimeIREERdg failed to find any module meta data, please reimport the original model"));
-		return {};
-	}
-
-	FModuleMetaData ModuleMetaData{};
-	{
-		FMemoryReaderView Reader(IREEModelData->ModuleMetaData, /*bIsPersitent =*/ true);
-		ModuleMetaData.Serialize(Reader);
-	}
-
-	if (ModuleMetaData.FunctionMetaData.IsEmpty())
-	{
-		UE_LOG(LogNNERuntimeIREE, Warning, TEXT("UNNERuntimeIREERdg failed to parse the module meta data, please reimport the original model"));
-		return {};
-	}
-
-	FIREECompilerRDGResult CompilerResult{};
-	{
-		FMemoryReaderView Reader(IREEModelData->CompilerResult, /*bIsPersitent =*/ true);
-		FIREECompilerRDGResult::StaticStruct()->SerializeBin(Reader, &CompilerResult);
-	}
 	
 	const FString ShaderPlatformName = LexToString(GShaderPlatformForFeatureLevel[GMaxRHIFeatureLevel]);
-	const FIREECompilerRDGBuildTargetResult* BuildTargetResult = CompilerResult.BuildTargetResults.FindByPredicate([ShaderPlatformName](const FIREECompilerRDGBuildTargetResult &Element) { return Element.ShaderPlatform == ShaderPlatformName; });
+	const FIREECompilerRDGBuildTargetResult* BuildTargetResult = IREEModelData->CompilerResult.BuildTargetResults.FindByPredicate([ShaderPlatformName](const FIREECompilerRDGBuildTargetResult &Element) { return Element.ShaderPlatform == ShaderPlatformName; });
 	if (!BuildTargetResult)
 	{
 		UE_LOG(LogNNERuntimeIREE, Warning, TEXT("UNNERuntimeIREERdg failed to find a matching shader platform for \'%s\'"), *ShaderPlatformName);
@@ -927,7 +855,7 @@ TSharedPtr<UE::NNE::IModelRDG> UNNERuntimeIREERdg::CreateModelRDG(const TObjectP
 	const FString SharedLibraryDirFullPath = FPaths::ConvertRelativePathToFull(FPaths::Combine(FPaths::ProjectDir(), GetSharedLibDirPath(UGameplayStatics::GetPlatformName(), GetRuntimeName(), FileIdString), RuntimeSubdir));
 	const FString SharedLibrarySubDirFullPath = FPaths::Combine(SharedLibraryDirFullPath, BuildTargetResult->ShaderPlatform);
 
-	TSharedPtr<UE::NNE::IModelRDG> Model = UE::NNERuntimeIREE::RDG::FModel::Make(SharedLibrarySubDirFullPath, BuildTargetResult->VmfbData, ModuleMetaData, ExecutableMap);
+	TSharedPtr<UE::NNE::IModelRDG> Model = UE::NNERuntimeIREE::RDG::FModel::Make(SharedLibrarySubDirFullPath, BuildTargetResult->VmfbData, IREEModelData->ModuleMetaData, ExecutableMap);
 	if (!Model.IsValid())
 	{
 		UE_LOG(LogNNERuntimeIREE, Warning, TEXT("UNNERuntimeIREERdg could not initialize the model created from model data with id %s"), *FileIdString);
diff --git a/Engine/Plugins/Experimental/NNERuntimeIREE/Source/NNERuntimeIREE/Private/NNERuntimeIREEModelData.cpp b/Engine/Plugins/Experimental/NNERuntimeIREE/Source/NNERuntimeIREE/Private/NNERuntimeIREEModelData.cpp
--- a/Engine/Plugins/Experimental/NNERuntimeIREE/Source/NNERuntimeIREE/Private/NNERuntimeIREEModelData.cpp
+++ b/Engine/Plugins/Experimental/NNERuntimeIREE/Source/NNERuntimeIREE/Private/NNERuntimeIREEModelData.cpp
@@ -10,8 +10,8 @@ void UNNERuntimeIREEModelDataCPU::Serialize(FArchive& Ar)
 	Ar << GUID;
 	Ar << Version;
 	Ar << FileId;
-	Ar << ModuleMetaData;
-	Ar << CompilerResult;
+	ModuleMetaData.Serialize(Ar);
+	FNNERuntimeIREECompilerResultCPU::StaticStruct()->SerializeBin(Ar, &CompilerResult);
 }
 
 bool UNNERuntimeIREEModelDataCPU::IsSameGuidAndVersion(TConstArrayView64<uint8> Data, FGuid Guid, int32 Version)
@@ -32,8 +32,8 @@ bool UNNERuntimeIREEModelDataCPU::IsSameGuidAndVersion(TConstArrayView64<uint8>
 void UNNERuntimeIREEModelDataRDG::Serialize(FArchive& Ar)
 {
 	FNNERuntimeIREEModelDataHeaderRDG::StaticStruct()->SerializeBin(Ar, &Header);
-	Ar << ModuleMetaData;
-	Ar << CompilerResult;
+	ModuleMetaData.Serialize(Ar);
+	FIREECompilerRDGResult::StaticStruct()->SerializeBin(Ar, &CompilerResult);
 }
 
 FNNERuntimeIREEModelDataHeaderRDG UNNERuntimeIREEModelDataRDG::ReadHeader(FArchive& Ar)
diff --git a/Engine/Plugins/Experimental/NNERuntimeIREE/Source/NNERuntimeIREE/Private/NNERuntimeIREEModelData.h b/Engine/Plugins/Experimental/NNERuntimeIREE/Source/NNERuntimeIREE/Private/NNERuntimeIREEModelData.h
--- a/Engine/Plugins/Experimental/NNERuntimeIREE/Source/NNERuntimeIREE/Private/NNERuntimeIREEModelData.h
+++ b/Engine/Plugins/Experimental/NNERuntimeIREE/Source/NNERuntimeIREE/Private/NNERuntimeIREEModelData.h
@@ -3,7 +3,9 @@
 #pragma once
 
 #include "Containers/ContainersFwd.h"
+#include "IREECompilerRDG.h"
 #include "NNERuntimeIREECompiler.h"
+#include "NNERuntimeIREEMetaData.h"
 #include "NNETypes.h"
 #include "UObject/Class.h"
 
@@ -45,12 +47,12 @@ public:
 	/**
 	 * Serialized module meta data.
 	 */
-	TArray64<uint8> ModuleMetaData;
+	UE::NNERuntimeIREE::FModuleMetaData ModuleMetaData;
 
 	/**
 	 * Serialized compiler output.
 	 */
-	TArray64<uint8> CompilerResult;
+	FNNERuntimeIREECompilerResultCPU CompilerResult;
 };
 
 /**
@@ -114,10 +116,10 @@ public:
 	/**
 	 * Serialized module meta data.
 	 */
-	TArray64<uint8> ModuleMetaData;
+	UE::NNERuntimeIREE::FModuleMetaData ModuleMetaData;
 
 	/**
 	 * Serialized compiler output.
 	 */
-	TArray64<uint8> CompilerResult;
+	FIREECompilerRDGResult CompilerResult;
 };
\ No newline at end of file
