[NNE] NNERuntimeIREE change ModelData to a simple class

#rnx
#tests NNEModelTest

#Synced-CL 49760974

diff --git a/Engine/Plugins/Experimental/NNERuntimeIREE/Source/NNERuntimeIREE/Private/NNERuntimeIREE.cpp b/Engine/Plugins/Experimental/NNERuntimeIREE/Source/NNERuntimeIREE/Private/NNERuntimeIREE.cpp
--- a/Engine/Plugins/Experimental/NNERuntimeIREE/Source/NNERuntimeIREE/Private/NNERuntimeIREE.cpp
+++ b/Engine/Plugins/Experimental/NNERuntimeIREE/Source/NNERuntimeIREE/Private/NNERuntimeIREE.cpp
@@ -245,6 +245,7 @@ TSharedPtr<UE::NNE::FSharedModelData> UNNERuntimeIREECpu::CreateModelData(const
 #if WITH_EDITOR
 	using namespace UE::NNERuntimeIREE;
 	using namespace UE::NNERuntimeIREE::Private;
+	using namespace UE::NNERuntimeIREE::CPU::Private;
 
 	const FString TargetPlatformName = GetTargetPlatformName(TargetPlatform);
 	if (CanCreateModelData(FileType, FileData, AdditionalFileData, FileId, TargetPlatform) != ECanCreateModelDataStatus::Ok)
@@ -264,26 +265,26 @@ TSharedPtr<UE::NNE::FSharedModelData> UNNERuntimeIREECpu::CreateModelData(const
 	FString IREEModelDataFilePath = FPaths::Combine(IntermediateDirFullPath, FileIdString) + ".ireemodeldata";
 
 	TArray64<uint8> ResultData;
-	TWeakObjectPtr<UNNERuntimeIREEModelDataCPU> IREEModelData = NewObject<UNNERuntimeIREEModelDataCPU>();
+	FModelData IREEModelData{};
 
 	bool bNeedCompileMlir = true;
 	if (PlatformFile.FileExists(*IREEModelDataFilePath) && FFileHelper::LoadFileToArray(ResultData, *IREEModelDataFilePath))
 	{
 		SCOPED_NAMED_EVENT_TEXT("Validate", FColor::Magenta);
 
-		if (UNNERuntimeIREEModelDataCPU::IsSameGuidAndVersion(ResultData, GUID, Version))
+		if (FModelData::IsSameGuidAndVersion(ResultData, GUID, Version))
 		{
 			{
 				FMemoryReaderView Reader(ResultData, /*bIsPersitent =*/ true);
-				IREEModelData->Serialize(Reader);
+				IREEModelData.Serialize(Reader);
 			}
 
-			check(FileIdString.Equals(IREEModelData->FileId.ToString(EGuidFormats::Digits).ToLower()));
+			check(FileIdString.Equals(IREEModelData.FileId.ToString(EGuidFormats::Digits).ToLower()));
 
 			bNeedCompileMlir = false;
-			for (int32 i = 0; i < IREEModelData->CompilerResult.ArchitectureInfos.Num(); i++)
+			for (int32 i = 0; i < IREEModelData.CompilerResult.ArchitectureInfos.Num(); i++)
 			{
-				const FNNERuntimeIREEArchitectureInfoCPU& Info = IREEModelData->CompilerResult.ArchitectureInfos[i];
+				const FNNERuntimeIREEArchitectureInfoCPU& Info = IREEModelData.CompilerResult.ArchitectureInfos[i];
 				const FString SharedLibrarySubDirFullPath = FPaths::Combine(SharedLibraryDirFullPath, Info.RelativeDirPath);
 
 				const FString SharedLibraryFilePath = FPaths::Combine(SharedLibrarySubDirFullPath, Info.SharedLibraryFileName);
@@ -322,23 +323,23 @@ TSharedPtr<UE::NNE::FSharedModelData> UNNERuntimeIREECpu::CreateModelData(const
 
 		// NOTE: From this point FileData refers to mlir data in any case
 
-		if (!Compiler->CompileMlir(FileData, FileIdString, IntermediateDirFullPath, IREEModelData->CompilerResult))
+		if (!Compiler->CompileMlir(FileData, FileIdString, IntermediateDirFullPath, IREEModelData.CompilerResult))
 		{
 			UE_LOG(LogNNERuntimeIREE, Warning, TEXT("UNNERuntimeIREECpu failed to compile model %s"), *FileIdString);
 			return TSharedPtr<UE::NNE::FSharedModelData>();
 		}
 
-		IREEModelData->GUID = UNNERuntimeIREECpu::GUID;
-		IREEModelData->Version = UNNERuntimeIREECpu::Version;
-		IREEModelData->FileId = FileId;
+		IREEModelData.GUID = UNNERuntimeIREECpu::GUID;
+		IREEModelData.Version = UNNERuntimeIREECpu::Version;
+		IREEModelData.FileId = FileId;
 		if (AdditionalFileData.Contains("IREEModuleMetaData") && !AdditionalFileData["IREEModuleMetaData"].IsEmpty())
 		{
 			FMemoryReaderView Reader(AdditionalFileData["IREEModuleMetaData"], /*bIsPersitent =*/ true);
-			IREEModelData->ModuleMetaData.Serialize(Reader);
+			IREEModelData.ModuleMetaData.Serialize(Reader);
 		}
 		else
 		{
-			if (!IREEModelData->ModuleMetaData.ParseFromBuffer(FileData))
+			if (!IREEModelData.ModuleMetaData.ParseFromBuffer(FileData))
 			{
 				UE_LOG(LogNNERuntimeIREE, Warning, TEXT("UNNERuntimeIREECpu failed to parse meta data %s"), *FileIdString);
 				return TSharedPtr<UE::NNE::FSharedModelData>();
@@ -347,7 +348,7 @@ TSharedPtr<UE::NNE::FSharedModelData> UNNERuntimeIREECpu::CreateModelData(const
 
 		{
 			FMemoryWriter64 Writer(ResultData, /*bIsPersitent =*/ true);
-			IREEModelData->Serialize(Writer);
+			IREEModelData.Serialize(Writer);
 		}
 
 		FFileHelper::SaveArrayToFile(ResultData, *IREEModelDataFilePath);
@@ -358,11 +359,11 @@ TSharedPtr<UE::NNE::FSharedModelData> UNNERuntimeIREECpu::CreateModelData(const
 	{
 		// Copy files for staging
 		FString StagingDirFullPath = FPaths::ConvertRelativePathToFull(FPaths::Combine(FPaths::ProjectDir(), GetPackagedModelDirPath(TargetPlatformName, GetRuntimeName(), FileIdString)));
-		for (int32 i = 0; i < IREEModelData->CompilerResult.ArchitectureInfos.Num(); i++)
+		for (int32 i = 0; i < IREEModelData.CompilerResult.ArchitectureInfos.Num(); i++)
 		{
 			SCOPED_NAMED_EVENT_TEXT("Copy", FColor::Magenta);
 
-			const FNNERuntimeIREEArchitectureInfoCPU& Info = IREEModelData->CompilerResult.ArchitectureInfos[i];
+			const FNNERuntimeIREEArchitectureInfoCPU& Info = IREEModelData.CompilerResult.ArchitectureInfos[i];
 			const FString SharedLibrarySubDirFullPath = FPaths::Combine(SharedLibraryDirFullPath, Info.RelativeDirPath);
 			const FString StagingSubDirFullPath = FPaths::Combine(StagingDirFullPath, Info.RelativeDirPath);
 			
@@ -394,13 +395,15 @@ UNNERuntimeIREECpu::ECanCreateModelCPUStatus UNNERuntimeIREECpu::CanCreateModelC
 {
 	check(ModelData != nullptr);
 
+	using namespace UE::NNERuntimeIREE::CPU::Private;
+
 	TSharedPtr<UE::NNE::FSharedModelData> SharedData = ModelData->GetModelData(GetRuntimeName());
 	if (!SharedData.IsValid())
 	{
 		return ECanCreateModelCPUStatus::Fail;
 	}
 
-	if (!UNNERuntimeIREEModelDataCPU::IsSameGuidAndVersion(SharedData->GetView(), UNNERuntimeIREECpu::GUID, UNNERuntimeIREECpu::Version))
+	if (!FModelData::IsSameGuidAndVersion(SharedData->GetView(), UNNERuntimeIREECpu::GUID, UNNERuntimeIREECpu::Version))
 	{
 		return ECanCreateModelCPUStatus::Fail;
 	}
@@ -414,6 +417,7 @@ TSharedPtr<UE::NNE::IModelCPU> UNNERuntimeIREECpu::CreateModelCPU(const TObjectP
 
 	using namespace UE::NNERuntimeIREE;
 	using namespace UE::NNERuntimeIREE::Private;
+	using namespace UE::NNERuntimeIREE::CPU::Private;
 
 	if (CanCreateModelCPU(ModelData) != ECanCreateModelCPUStatus::Ok)
 	{
@@ -426,18 +430,18 @@ TSharedPtr<UE::NNE::IModelCPU> UNNERuntimeIREECpu::CreateModelCPU(const TObjectP
 
 	TConstArrayView<uint8> SharedDataView = SharedData->GetView();
 
-	TWeakObjectPtr<UNNERuntimeIREEModelDataCPU> IREEModelData = NewObject<UNNERuntimeIREEModelDataCPU>();
+	FModelData IREEModelData{};
 	{
 		FMemoryReaderView Reader(SharedDataView, /*bIsPersitent =*/ true);
-		IREEModelData->Serialize(Reader);
+		IREEModelData.Serialize(Reader);
 	}
 
 	int32 TargetIndex = -1;
-	for (int32 i = 0; i < IREEModelData->CompilerResult.ArchitectureInfos.Num(); i++)
+	for (int32 i = 0; i < IREEModelData.CompilerResult.ArchitectureInfos.Num(); i++)
 	{
 		bool bMatches = true;
-		bMatches &= DoesArchitectureMatch(IREEModelData->CompilerResult.ArchitectureInfos[i].Architecture);
-		bMatches &= DoX86FeaturesMatch(IREEModelData->CompilerResult.ArchitectureInfos[i].X86Features);
+		bMatches &= DoesArchitectureMatch(IREEModelData.CompilerResult.ArchitectureInfos[i].Architecture);
+		bMatches &= DoX86FeaturesMatch(IREEModelData.CompilerResult.ArchitectureInfos[i].X86Features);
 
 		if (bMatches)
 		{
@@ -451,12 +455,12 @@ TSharedPtr<UE::NNE::IModelCPU> UNNERuntimeIREECpu::CreateModelCPU(const TObjectP
 		return TSharedPtr<UE::NNE::IModelCPU>();
 	}
 
-	const FNNERuntimeIREEArchitectureInfoCPU& ArchitectureInfo = IREEModelData->CompilerResult.ArchitectureInfos[TargetIndex];
-	const FString FileIdString = IREEModelData->FileId.ToString(EGuidFormats::Digits).ToLower();
+	const FNNERuntimeIREEArchitectureInfoCPU& ArchitectureInfo = IREEModelData.CompilerResult.ArchitectureInfos[TargetIndex];
+	const FString FileIdString = IREEModelData.FileId.ToString(EGuidFormats::Digits).ToLower();
 	const FString SharedLibraryDirFullPath = FPaths::ConvertRelativePathToFull(FPaths::Combine(FPaths::ProjectDir(), GetSharedLibDirPath(UGameplayStatics::GetPlatformName(), GetRuntimeName(), FileIdString)));
 	const FString SharedLibrarySubDirFullPath = FPaths::Combine(SharedLibraryDirFullPath, ArchitectureInfo.RelativeDirPath);
 
-	TSharedPtr<UE::NNE::IModelCPU> Model = UE::NNERuntimeIREE::CPU::FModel::Make(*Environment, SharedLibrarySubDirFullPath, ArchitectureInfo.SharedLibraryFileName, ArchitectureInfo.VmfbFileName, ArchitectureInfo.SharedLibraryEntryPointName, IREEModelData->ModuleMetaData);
+	TSharedPtr<UE::NNE::IModelCPU> Model = UE::NNERuntimeIREE::CPU::FModel::Make(*Environment, SharedLibrarySubDirFullPath, ArchitectureInfo.SharedLibraryFileName, ArchitectureInfo.VmfbFileName, ArchitectureInfo.SharedLibraryEntryPointName, IREEModelData.ModuleMetaData);
 	if (!Model.IsValid())
 	{
 		UE_LOG(LogNNERuntimeIREE, Warning, TEXT("UNNERuntimeIREECpu could not initialize the model created from model data with id %s"), *FileIdString);
@@ -635,6 +639,7 @@ TSharedPtr<UE::NNE::FSharedModelData> UNNERuntimeIREERdg::CreateModelData(const
 #if WITH_EDITOR
 	using namespace UE::NNERuntimeIREE;
 	using namespace UE::NNERuntimeIREE::Private;
+	using namespace UE::NNERuntimeIREE::RDG::Private;
 
 	const FString TargetPlatformName = GetTargetPlatformName(TargetPlatform);
 	if (CanCreateModelData(FileType, FileData, AdditionalFileData, FileId, TargetPlatform) != ECanCreateModelDataStatus::Ok)
@@ -662,12 +667,12 @@ TSharedPtr<UE::NNE::FSharedModelData> UNNERuntimeIREERdg::CreateModelData(const
 		SCOPED_NAMED_EVENT_TEXT("Validate", FColor::Magenta);
 
 		TUniquePtr<FArchive> FileReader(IFileManager::Get().CreateFileReader(*IREEModelDataFilePath));
-		if (FileReader && UNNERuntimeIREEModelDataRDG::IsSameGuidAndVersion(*FileReader, GUID, Version))
+		if (FileReader && FModelData::IsSameGuidAndVersion(*FileReader, GUID, Version))
 		{
 			FileReader->Seek(0);
 			bNeedCompileMlir = false;
 
-			FNNERuntimeIREEModelDataHeaderRDG Header = UNNERuntimeIREEModelDataRDG::ReadHeader(*FileReader);
+			FNNERuntimeIREEModelDataHeaderRDG Header = FModelData::ReadHeader(*FileReader);
 
 			check(FileIdString.Equals(Header.FileId.ToString(EGuidFormats::Digits).ToLower()));
 
@@ -711,30 +716,30 @@ TSharedPtr<UE::NNE::FSharedModelData> UNNERuntimeIREERdg::CreateModelData(const
 
 		// NOTE: From this point FileData refers to mlir data in any case
 
-		TWeakObjectPtr<UNNERuntimeIREEModelDataRDG> IREEModelData = NewObject<UNNERuntimeIREEModelDataRDG>();
-		if (!Compiler->CompileMlir(FileData, FileIdString, IntermediateDirFullPath, ShaderPlatforms, IREEModelData->CompilerResult))
+		FModelData IREEModelData{};
+		if (!Compiler->CompileMlir(FileData, FileIdString, IntermediateDirFullPath, ShaderPlatforms, IREEModelData.CompilerResult))
 		{
 			UE_LOG(LogNNERuntimeIREE, Warning, TEXT("UNNERuntimeIREERdg failed to compile model %s"), *FileIdString);
 			return TSharedPtr<UE::NNE::FSharedModelData>();
 		}
 
-		IREEModelData->Header.GUID = UNNERuntimeIREERdg::GUID;
-		IREEModelData->Header.Version = UNNERuntimeIREERdg::Version;
-		IREEModelData->Header.FileId = FileId;
+		IREEModelData.Header.GUID = UNNERuntimeIREERdg::GUID;
+		IREEModelData.Header.Version = UNNERuntimeIREERdg::Version;
+		IREEModelData.Header.FileId = FileId;
 		for (EShaderPlatform ShaderPlatform : ShaderPlatforms)
 		{
 			const FString ShaderPlatformName = LexToString(ShaderPlatform);
 
-			IREEModelData->Header.ShaderPlatforms.Add(ShaderPlatformName);
+			IREEModelData.Header.ShaderPlatforms.Add(ShaderPlatformName);
 		}
 		if (AdditionalFileData.Contains("IREEModuleMetaData") && !AdditionalFileData["IREEModuleMetaData"].IsEmpty())
 		{
 			FMemoryReaderView Reader(AdditionalFileData["IREEModuleMetaData"]);
-			IREEModelData->ModuleMetaData.Serialize(Reader);
+			IREEModelData.ModuleMetaData.Serialize(Reader);
 		}
 		else
 		{
-			if (!IREEModelData->ModuleMetaData.ParseFromBuffer(FileData))
+			if (!IREEModelData.ModuleMetaData.ParseFromBuffer(FileData))
 			{
 				UE_LOG(LogNNERuntimeIREE, Warning, TEXT("UNNERuntimeIREERdg failed to parse meta data %s"), *FileIdString);
 				return TSharedPtr<UE::NNE::FSharedModelData>();
@@ -743,7 +748,7 @@ TSharedPtr<UE::NNE::FSharedModelData> UNNERuntimeIREERdg::CreateModelData(const
 
 		{
 			FMemoryWriter64 Writer(ResultData, /*bIsPersitent =*/ true);
-			IREEModelData->Serialize(Writer);
+			IREEModelData.Serialize(Writer);
 		}
 
 		if (!FFileHelper::SaveArrayToFile(ResultData, *IREEModelDataFilePath))
@@ -777,6 +782,8 @@ UNNERuntimeIREERdg::ECanCreateModelRDGStatus UNNERuntimeIREERdg::CanCreateModelR
 {
 	check(ModelData != nullptr);
 
+	using namespace UE::NNERuntimeIREE::RDG::Private;
+
 	TSharedPtr<UE::NNE::FSharedModelData> SharedData = ModelData->GetModelData(GetRuntimeName());
 	if (!SharedData.IsValid())
 	{
@@ -784,12 +791,12 @@ UNNERuntimeIREERdg::ECanCreateModelRDGStatus UNNERuntimeIREERdg::CanCreateModelR
 	}
 
 	FMemoryReaderView Reader(SharedData->GetView(), /*bIsPersitent =*/ true);
-	if (!UNNERuntimeIREEModelDataRDG::IsSameGuidAndVersion(Reader, GUID, Version))
+	if (!FModelData::IsSameGuidAndVersion(Reader, GUID, Version))
 	{
 		return ECanCreateModelRDGStatus::Fail;
 	}
 	Reader.Seek(0);
-	FNNERuntimeIREEModelDataHeaderRDG Header = UNNERuntimeIREEModelDataRDG::ReadHeader(Reader);
+	FNNERuntimeIREEModelDataHeaderRDG Header = FModelData::ReadHeader(Reader);
 
 	const FString ShaderPlatformName = LexToString(GShaderPlatformForFeatureLevel[GMaxRHIFeatureLevel]);
 	if (Header.ShaderPlatforms.Find(ShaderPlatformName) == INDEX_NONE)
@@ -806,6 +813,7 @@ TSharedPtr<UE::NNE::IModelRDG> UNNERuntimeIREERdg::CreateModelRDG(const TObjectP
 
 	using namespace UE::NNERuntimeIREE;
 	using namespace UE::NNERuntimeIREE::Private;
+	using namespace UE::NNERuntimeIREE::RDG::Private;
 
 	if (CanCreateModelRDG(ModelData) != ECanCreateModelRDGStatus::Ok)
 	{
@@ -818,14 +826,14 @@ TSharedPtr<UE::NNE::IModelRDG> UNNERuntimeIREERdg::CreateModelRDG(const TObjectP
 
 	TConstArrayView<uint8> SharedDataView = SharedData->GetView();
 
-	TWeakObjectPtr<UNNERuntimeIREEModelDataRDG> IREEModelData = NewObject<UNNERuntimeIREEModelDataRDG>();
+	FModelData IREEModelData{};
 	{
 		FMemoryReaderView Reader(SharedDataView, /*bIsPersitent =*/ true);
-		IREEModelData->Serialize(Reader);
+		IREEModelData.Serialize(Reader);
 	}
 	
 	const FString ShaderPlatformName = LexToString(GShaderPlatformForFeatureLevel[GMaxRHIFeatureLevel]);
-	const FIREECompilerRDGBuildTargetResult* BuildTargetResult = IREEModelData->CompilerResult.BuildTargetResults.FindByPredicate([ShaderPlatformName](const FIREECompilerRDGBuildTargetResult &Element) { return Element.ShaderPlatform == ShaderPlatformName; });
+	const FIREECompilerRDGBuildTargetResult* BuildTargetResult = IREEModelData.CompilerResult.BuildTargetResults.FindByPredicate([ShaderPlatformName](const FIREECompilerRDGBuildTargetResult &Element) { return Element.ShaderPlatform == ShaderPlatformName; });
 	if (!BuildTargetResult)
 	{
 		UE_LOG(LogNNERuntimeIREE, Warning, TEXT("UNNERuntimeIREERdg failed to find a matching shader platform for \'%s\'"), *ShaderPlatformName);
@@ -850,12 +858,12 @@ TSharedPtr<UE::NNE::IModelRDG> UNNERuntimeIREERdg::CreateModelRDG(const TObjectP
 		ExecutableMap.Emplace(ExecutableData.Name, ExecutableData.Data);
 	}
 
-	const FString FileIdString = IREEModelData->Header.FileId.ToString(EGuidFormats::Digits).ToLower();
+	const FString FileIdString = IREEModelData.Header.FileId.ToString(EGuidFormats::Digits).ToLower();
 	const FString RuntimeSubdir = GetRuntimeSubdir();
 	const FString SharedLibraryDirFullPath = FPaths::ConvertRelativePathToFull(FPaths::Combine(FPaths::ProjectDir(), GetSharedLibDirPath(UGameplayStatics::GetPlatformName(), GetRuntimeName(), FileIdString), RuntimeSubdir));
 	const FString SharedLibrarySubDirFullPath = FPaths::Combine(SharedLibraryDirFullPath, BuildTargetResult->ShaderPlatform);
 
-	TSharedPtr<UE::NNE::IModelRDG> Model = UE::NNERuntimeIREE::RDG::FModel::Make(SharedLibrarySubDirFullPath, BuildTargetResult->VmfbData, IREEModelData->ModuleMetaData, ExecutableMap);
+	TSharedPtr<UE::NNE::IModelRDG> Model = UE::NNERuntimeIREE::RDG::FModel::Make(SharedLibrarySubDirFullPath, BuildTargetResult->VmfbData, IREEModelData.ModuleMetaData, ExecutableMap);
 	if (!Model.IsValid())
 	{
 		UE_LOG(LogNNERuntimeIREE, Warning, TEXT("UNNERuntimeIREERdg could not initialize the model created from model data with id %s"), *FileIdString);
diff --git a/Engine/Plugins/Experimental/NNERuntimeIREE/Source/NNERuntimeIREE/Private/NNERuntimeIREEModelData.cpp b/Engine/Plugins/Experimental/NNERuntimeIREE/Source/NNERuntimeIREE/Private/NNERuntimeIREEModelData.cpp
--- a/Engine/Plugins/Experimental/NNERuntimeIREE/Source/NNERuntimeIREE/Private/NNERuntimeIREEModelData.cpp
+++ b/Engine/Plugins/Experimental/NNERuntimeIREE/Source/NNERuntimeIREE/Private/NNERuntimeIREEModelData.cpp
@@ -5,46 +5,55 @@
 #include "Containers/Array.h"
 #include "NNE.h"
 
-void UNNERuntimeIREEModelDataCPU::Serialize(FArchive& Ar)
+namespace UE::NNERuntimeIREE
 {
-	Ar << GUID;
-	Ar << Version;
-	Ar << FileId;
-	ModuleMetaData.Serialize(Ar);
-	FNNERuntimeIREECompilerResultCPU::StaticStruct()->SerializeBin(Ar, &CompilerResult);
-}
-
-bool UNNERuntimeIREEModelDataCPU::IsSameGuidAndVersion(TConstArrayView64<uint8> Data, FGuid Guid, int32 Version)
-{
-	int32 GuidSize = sizeof(UNNERuntimeIREEModelDataCPU::GUID);
-	int32 VersionSize = sizeof(UNNERuntimeIREEModelDataCPU::Version);
-	if (Data.Num() <= GuidSize + VersionSize)
+	namespace CPU::Private
 	{
-		return false;
-	}
-
-	bool bResult = FGenericPlatformMemory::Memcmp(&(Data[0]), &Guid, GuidSize) == 0;
-	bResult &= FGenericPlatformMemory::Memcmp(&(Data[GuidSize]), &Version, VersionSize) == 0;
-
-	return bResult;
-}
-
-void UNNERuntimeIREEModelDataRDG::Serialize(FArchive& Ar)
-{
-	FNNERuntimeIREEModelDataHeaderRDG::StaticStruct()->SerializeBin(Ar, &Header);
-	ModuleMetaData.Serialize(Ar);
-	FIREECompilerRDGResult::StaticStruct()->SerializeBin(Ar, &CompilerResult);
-}
-
-FNNERuntimeIREEModelDataHeaderRDG UNNERuntimeIREEModelDataRDG::ReadHeader(FArchive& Ar)
-{
-	FNNERuntimeIREEModelDataHeaderRDG Header{};
-	FNNERuntimeIREEModelDataHeaderRDG::StaticStruct()->SerializeBin(Ar, &Header);
-	return Header;
-}
-
-bool UNNERuntimeIREEModelDataRDG::IsSameGuidAndVersion(FArchive& Ar, FGuid Guid, int32 Version)
-{
-	FNNERuntimeIREEModelDataHeaderRDG Header = ReadHeader(Ar);
-	return Guid == Header.GUID && Version == Header.Version && !Ar.IsError();
-}
\ No newline at end of file
+		void FModelData::Serialize(FArchive& Ar)
+		{
+			Ar << GUID;
+			Ar << Version;
+			Ar << FileId;
+			ModuleMetaData.Serialize(Ar);
+			FNNERuntimeIREECompilerResultCPU::StaticStruct()->SerializeBin(Ar, &CompilerResult);
+		}
+
+		bool FModelData::IsSameGuidAndVersion(TConstArrayView64<uint8> Data, FGuid Guid, int32 Version)
+		{
+			int32 GuidSize = sizeof(FModelData::GUID);
+			int32 VersionSize = sizeof(FModelData::Version);
+			if (Data.Num() <= GuidSize + VersionSize)
+			{
+				return false;
+			}
+
+			bool bResult = FGenericPlatformMemory::Memcmp(&(Data[0]), &Guid, GuidSize) == 0;
+			bResult &= FGenericPlatformMemory::Memcmp(&(Data[GuidSize]), &Version, VersionSize) == 0;
+
+			return bResult;
+		}
+	} // CPU::Private
+
+	namespace RDG::Private
+	{
+		void FModelData::Serialize(FArchive& Ar)
+		{
+			FNNERuntimeIREEModelDataHeaderRDG::StaticStruct()->SerializeBin(Ar, &Header);
+			ModuleMetaData.Serialize(Ar);
+			FIREECompilerRDGResult::StaticStruct()->SerializeBin(Ar, &CompilerResult);
+		}
+
+		FNNERuntimeIREEModelDataHeaderRDG FModelData::ReadHeader(FArchive& Ar)
+		{
+			FNNERuntimeIREEModelDataHeaderRDG Header{};
+			FNNERuntimeIREEModelDataHeaderRDG::StaticStruct()->SerializeBin(Ar, &Header);
+			return Header;
+		}
+
+		bool FModelData::IsSameGuidAndVersion(FArchive& Ar, FGuid Guid, int32 Version)
+		{
+			FNNERuntimeIREEModelDataHeaderRDG Header = ReadHeader(Ar);
+			return Guid == Header.GUID && Version == Header.Version && !Ar.IsError();
+		}
+	} // RDG::Private
+} //UE::NNERuntimeIREE
diff --git a/Engine/Plugins/Experimental/NNERuntimeIREE/Source/NNERuntimeIREE/Private/NNERuntimeIREEModelData.h b/Engine/Plugins/Experimental/NNERuntimeIREE/Source/NNERuntimeIREE/Private/NNERuntimeIREEModelData.h
--- a/Engine/Plugins/Experimental/NNERuntimeIREE/Source/NNERuntimeIREE/Private/NNERuntimeIREEModelData.h
+++ b/Engine/Plugins/Experimental/NNERuntimeIREE/Source/NNERuntimeIREE/Private/NNERuntimeIREEModelData.h
@@ -11,49 +11,47 @@
 
 #include "NNERuntimeIREEModelData.generated.h"
 
-/**
- * IREE model data class.
- */
-UCLASS()
-class UNNERuntimeIREEModelDataCPU : public UObject
+namespace UE::NNERuntimeIREE::CPU::Private
 {
-	GENERATED_BODY()
-
-public:
-	// UObject interface
-	virtual void Serialize(FArchive& Ar) override;
-	// End of UObject interface
-
 	/**
-	 * Check data for matching Guid and Version without deserializing everything.
+	 * IREE model data class.
 	 */
-	static bool IsSameGuidAndVersion(TConstArrayView64<uint8> Data, FGuid Guid, int32 Version);
-
-	/**
-	 * A Guid that uniquely identifies this IREE model data.
-	 */
-	FGuid GUID;
-
-	/**
-	 * Current version of this IREE model data.
-	 */
-	int32 Version;
-
-	/**
-	 * A Guid that uniquely identifies the model.
-	 */
-	FGuid FileId;
-
-	/**
-	 * Serialized module meta data.
-	 */
-	UE::NNERuntimeIREE::FModuleMetaData ModuleMetaData;
-
-	/**
-	 * Serialized compiler output.
-	 */
-	FNNERuntimeIREECompilerResultCPU CompilerResult;
-};
+	class FModelData
+	{
+	public:
+		void Serialize(FArchive& Ar);
+
+		/**
+		 * Check data for matching Guid and Version without deserializing everything.
+		 */
+		static bool IsSameGuidAndVersion(TConstArrayView64<uint8> Data, FGuid Guid, int32 Version);
+
+		/**
+		 * A Guid that uniquely identifies this IREE model data.
+		 */
+		FGuid GUID;
+
+		/**
+		 * Current version of this IREE model data.
+		 */
+		int32 Version;
+
+		/**
+		 * A Guid that uniquely identifies the model.
+		 */
+		FGuid FileId;
+
+		/**
+		 * Module meta data.
+		 */
+		UE::NNERuntimeIREE::FModuleMetaData ModuleMetaData;
+
+		/**
+		 * Compiler output.
+		 */
+		FNNERuntimeIREECompilerResultCPU CompilerResult;
+	};
+} // UE::NNERuntimeIREE::CPU::Private
 
 /**
  * IREE model data header RDG struct.
@@ -88,38 +86,36 @@ struct FNNERuntimeIREEModelDataHeaderRDG
 	TArray<FString> ShaderPlatforms;
 };
 
-/**
- * IREE model data RDG class.
- */
-UCLASS()
-class UNNERuntimeIREEModelDataRDG : public UObject
+namespace UE::NNERuntimeIREE::RDG::Private
 {
-	GENERATED_BODY()
-
-public:
-	// UObject interface
-	virtual void Serialize(FArchive& Ar) override;
-	// End of UObject interface
-
-	/**
-	 * Check data for matching Guid and Version with only deserializing the header
-	 */
-	static bool IsSameGuidAndVersion(FArchive& Ar, FGuid Guid, int32 Version);
-
-	static FNNERuntimeIREEModelDataHeaderRDG ReadHeader(FArchive& Ar);
-
-	/**
-	 * Model data header.
-	 */
-	FNNERuntimeIREEModelDataHeaderRDG Header;
-
-	/**
-	 * Serialized module meta data.
-	 */
-	UE::NNERuntimeIREE::FModuleMetaData ModuleMetaData;
-
 	/**
-	 * Serialized compiler output.
+	 * IREE model data RDG class.
 	 */
-	FIREECompilerRDGResult CompilerResult;
-};
\ No newline at end of file
+	class FModelData
+	{
+	public:
+		void Serialize(FArchive& Ar);
+
+		/**
+		 * Check data for matching Guid and Version with only deserializing the header
+		 */
+		static bool IsSameGuidAndVersion(FArchive& Ar, FGuid Guid, int32 Version);
+
+		static FNNERuntimeIREEModelDataHeaderRDG ReadHeader(FArchive& Ar);
+
+		/**
+		 * Model data header.
+		 */
+		FNNERuntimeIREEModelDataHeaderRDG Header;
+
+		/**
+		 * Module meta data.
+		 */
+		UE::NNERuntimeIREE::FModuleMetaData ModuleMetaData;
+
+		/**
+		 * Compiler output.
+		 */
+		FIREECompilerRDGResult CompilerResult;
+	};
+} // UE::NNERuntimeIREE::RDG::Private
diff --git a/Engine/Plugins/Experimental/NNERuntimeIREE/Source/NNERuntimeIREE/Private/NNERuntimeIREEModule.h b/Engine/Plugins/Experimental/NNERuntimeIREE/Source/NNERuntimeIREE/Private/NNERuntimeIREEModule.h
--- a/Engine/Plugins/Experimental/NNERuntimeIREE/Source/NNERuntimeIREE/Private/NNERuntimeIREEModule.h
+++ b/Engine/Plugins/Experimental/NNERuntimeIREE/Source/NNERuntimeIREE/Private/NNERuntimeIREEModule.h
@@ -8,11 +8,6 @@
 #include "NNERuntimeIREE.h"
 #include "UObject/WeakObjectPtrTemplates.h"
 
-namespace UE::NNERuntimeORT::Private
-{
-	class FEnvironment;
-}
-
 #endif // WITH_NNE_RUNTIME_IREE
 
 class FNNERuntimeIREEModule : public IModuleInterface
