[NNE] NNERuntimeIREE improved error handling

#rnx
#tests NNEModelTest

Checks for FArchive::IsError() after serialization and checks for file operation success.

#Synced-CL 49762123

diff --git a/Engine/Plugins/Experimental/NNERuntimeIREE/Source/NNERuntimeIREE/Private/NNERuntimeIREE.cpp b/Engine/Plugins/Experimental/NNERuntimeIREE/Source/NNERuntimeIREE/Private/NNERuntimeIREE.cpp
--- a/Engine/Plugins/Experimental/NNERuntimeIREE/Source/NNERuntimeIREE/Private/NNERuntimeIREE.cpp
+++ b/Engine/Plugins/Experimental/NNERuntimeIREE/Source/NNERuntimeIREE/Private/NNERuntimeIREE.cpp
@@ -272,13 +272,8 @@ TSharedPtr<UE::NNE::FSharedModelData> UNNERuntimeIREECpu::CreateModelData(const
 	{
 		SCOPED_NAMED_EVENT_TEXT("Validate", FColor::Magenta);
 
-		if (FModelData::IsSameGuidAndVersion(ResultData, GUID, Version))
+		if (FModelData::IsSameGuidAndVersion(ResultData, GUID, Version) && IREEModelData.Load(ResultData))
 		{
-			{
-				FMemoryReaderView Reader(ResultData, /*bIsPersitent =*/ true);
-				IREEModelData.Serialize(Reader);
-			}
-
 			check(FileIdString.Equals(IREEModelData.FileId.ToString(EGuidFormats::Digits).ToLower()));
 
 			bNeedCompileMlir = false;
@@ -346,12 +341,11 @@ TSharedPtr<UE::NNE::FSharedModelData> UNNERuntimeIREECpu::CreateModelData(const
 			}
 		}
 
+		if (!IREEModelData.Store(ResultData) || !FFileHelper::SaveArrayToFile(ResultData, *IREEModelDataFilePath))
 		{
-			FMemoryWriter64 Writer(ResultData, /*bIsPersitent =*/ true);
-			IREEModelData.Serialize(Writer);
+			UE_LOG(LogNNERuntimeIREE, Warning, TEXT("UNNERuntimeIREECpu failed to write model data to disk %s"), *IREEModelDataFilePath);
+			return TSharedPtr<UE::NNE::FSharedModelData>();
 		}
-
-		FFileHelper::SaveArrayToFile(ResultData, *IREEModelDataFilePath);
 	}
 
 	// Only stage when cooking
@@ -428,12 +422,11 @@ TSharedPtr<UE::NNE::IModelCPU> UNNERuntimeIREECpu::CreateModelCPU(const TObjectP
 	TSharedPtr<UE::NNE::FSharedModelData> SharedData = ModelData->GetModelData(GetRuntimeName());
 	check(SharedData.IsValid());
 
-	TConstArrayView<uint8> SharedDataView = SharedData->GetView();
-
 	FModelData IREEModelData{};
+	if(!IREEModelData.Load(SharedData->GetView()))
 	{
-		FMemoryReaderView Reader(SharedDataView, /*bIsPersitent =*/ true);
-		IREEModelData.Serialize(Reader);
+		UE_LOG(LogNNERuntimeIREE, Warning, TEXT("UNNERuntimeIREECpu failed to load model data %s"), *ModelData->GetFileId().ToString(EGuidFormats::Digits));
+		return TSharedPtr<UE::NNE::IModelCPU>();
 	}
 
 	int32 TargetIndex = -1;
@@ -666,19 +659,13 @@ TSharedPtr<UE::NNE::FSharedModelData> UNNERuntimeIREERdg::CreateModelData(const
 	{
 		SCOPED_NAMED_EVENT_TEXT("Validate", FColor::Magenta);
 
+		FNNERuntimeIREEModelDataHeaderRDG Header{};
 		TUniquePtr<FArchive> FileReader(IFileManager::Get().CreateFileReader(*IREEModelDataFilePath));
-		if (FileReader && FModelData::IsSameGuidAndVersion(*FileReader, GUID, Version))
+		if (FileReader && FModelData::ReadHeaderAndIsSameGuidAndVersion(*FileReader, GUID, Version, Header))
 		{
-			FileReader->Seek(0);
-			bNeedCompileMlir = false;
-
-			FNNERuntimeIREEModelDataHeaderRDG Header = FModelData::ReadHeader(*FileReader);
-
 			check(FileIdString.Equals(Header.FileId.ToString(EGuidFormats::Digits).ToLower()));
 
-			bNeedCompileMlir |= Header.GUID != GUID;
-			bNeedCompileMlir |= Header.Version != Version;
-
+			bNeedCompileMlir = false;
 			for (EShaderPlatform ShaderPlatform : ShaderPlatforms)
 			{
 				const FString ShaderPlatformName = LexToString(ShaderPlatform);
@@ -746,15 +733,10 @@ TSharedPtr<UE::NNE::FSharedModelData> UNNERuntimeIREERdg::CreateModelData(const
 			}
 		}
 
-		{
-			FMemoryWriter64 Writer(ResultData, /*bIsPersitent =*/ true);
-			IREEModelData.Serialize(Writer);
-		}
-
-		if (!FFileHelper::SaveArrayToFile(ResultData, *IREEModelDataFilePath))
+		if (!IREEModelData.Store(ResultData) || !FFileHelper::SaveArrayToFile(ResultData, *IREEModelDataFilePath))
 		{
 			UE_LOG(LogNNERuntimeIREE, Warning, TEXT("UNNERuntimeIREERdg failed to write model data to disk %s"), *IREEModelDataFilePath);
-			return {};
+			return TSharedPtr<UE::NNE::FSharedModelData>();
 		}
 	}
 	else
@@ -762,13 +744,13 @@ TSharedPtr<UE::NNE::FSharedModelData> UNNERuntimeIREERdg::CreateModelData(const
 		if (!FFileHelper::LoadFileToArray(ResultData, *IREEModelDataFilePath))
 		{
 			UE_LOG(LogNNERuntimeIREE, Warning, TEXT("UNNERuntimeIREERdg failed to read model data from disk %s"), *IREEModelDataFilePath);
-			return {};
+			return TSharedPtr<UE::NNE::FSharedModelData>();
 		}
 	}
 
 	return MakeShared<UE::NNE::FSharedModelData>(MakeSharedBufferFromArray(MoveTemp(ResultData)), 0);
 #else
-	return {};
+	return TSharedPtr<UE::NNE::FSharedModelData>();
 #endif // WITH_EDITOR
 }
 
@@ -790,13 +772,12 @@ UNNERuntimeIREERdg::ECanCreateModelRDGStatus UNNERuntimeIREERdg::CanCreateModelR
 		return ECanCreateModelRDGStatus::Fail;
 	}
 
+	FNNERuntimeIREEModelDataHeaderRDG Header{};
 	FMemoryReaderView Reader(SharedData->GetView(), /*bIsPersitent =*/ true);
-	if (!FModelData::IsSameGuidAndVersion(Reader, GUID, Version))
+	if (!FModelData::ReadHeaderAndIsSameGuidAndVersion(Reader, GUID, Version, Header))
 	{
 		return ECanCreateModelRDGStatus::Fail;
 	}
-	Reader.Seek(0);
-	FNNERuntimeIREEModelDataHeaderRDG Header = FModelData::ReadHeader(Reader);
 
 	const FString ShaderPlatformName = LexToString(GShaderPlatformForFeatureLevel[GMaxRHIFeatureLevel]);
 	if (Header.ShaderPlatforms.Find(ShaderPlatformName) == INDEX_NONE)
@@ -824,12 +805,11 @@ TSharedPtr<UE::NNE::IModelRDG> UNNERuntimeIREERdg::CreateModelRDG(const TObjectP
 	TSharedPtr<UE::NNE::FSharedModelData> SharedData = ModelData->GetModelData(GetRuntimeName());
 	check(SharedData.IsValid());
 
-	TConstArrayView<uint8> SharedDataView = SharedData->GetView();
-
 	FModelData IREEModelData{};
+	if (!IREEModelData.Load(SharedData->GetView()))
 	{
-		FMemoryReaderView Reader(SharedDataView, /*bIsPersitent =*/ true);
-		IREEModelData.Serialize(Reader);
+		UE_LOG(LogNNERuntimeIREE, Warning, TEXT("UNNERuntimeIREERdg failed to load model data %s"), *ModelData->GetFileId().ToString(EGuidFormats::Digits));
+		return TSharedPtr<UE::NNE::IModelRDG>();
 	}
 	
 	const FString ShaderPlatformName = LexToString(GShaderPlatformForFeatureLevel[GMaxRHIFeatureLevel]);
diff --git a/Engine/Plugins/Experimental/NNERuntimeIREE/Source/NNERuntimeIREE/Private/NNERuntimeIREEModelData.cpp b/Engine/Plugins/Experimental/NNERuntimeIREE/Source/NNERuntimeIREE/Private/NNERuntimeIREEModelData.cpp
--- a/Engine/Plugins/Experimental/NNERuntimeIREE/Source/NNERuntimeIREE/Private/NNERuntimeIREEModelData.cpp
+++ b/Engine/Plugins/Experimental/NNERuntimeIREE/Source/NNERuntimeIREE/Private/NNERuntimeIREEModelData.cpp
@@ -4,6 +4,8 @@
 
 #include "Containers/Array.h"
 #include "NNE.h"
+#include "Serialization/MemoryReader.h"
+#include "Serialization/MemoryWriter.h"
 
 namespace UE::NNERuntimeIREE
 {
@@ -18,19 +20,27 @@ namespace UE::NNERuntimeIREE
 			FNNERuntimeIREECompilerResultCPU::StaticStruct()->SerializeBin(Ar, &CompilerResult);
 		}
 
-		bool FModelData::IsSameGuidAndVersion(TConstArrayView64<uint8> Data, FGuid Guid, int32 Version)
+		bool FModelData::Load(TConstArrayView64<uint8> Data)
 		{
-			int32 GuidSize = sizeof(FModelData::GUID);
-			int32 VersionSize = sizeof(FModelData::Version);
-			if (Data.Num() <= GuidSize + VersionSize)
-			{
-				return false;
-			}
+			FMemoryReaderView Reader(Data, /*bIsPersitent =*/ true);
+			Serialize(Reader);
+			return !Reader.IsError();
+		}
 
-			bool bResult = FGenericPlatformMemory::Memcmp(&(Data[0]), &Guid, GuidSize) == 0;
-			bResult &= FGenericPlatformMemory::Memcmp(&(Data[GuidSize]), &Version, VersionSize) == 0;
+		bool FModelData::Store(TArray64<uint8>& Data)
+		{
+			FMemoryWriter64 Writer(Data, /*bIsPersitent =*/ true);
+			Serialize(Writer);
+			return !Writer.IsError();
+		}
 
-			return bResult;
+		bool FModelData::IsSameGuidAndVersion(TConstArrayView64<uint8> Data, FGuid Guid, int32 Version)
+		{
+			FModelData ModelData{};
+			FMemoryReaderView Reader(Data, /*bIsPersitent =*/ true);
+			Reader << ModelData.GUID;
+			Reader << ModelData.Version;
+			return !Reader.IsError() && Guid == ModelData.GUID && Version == ModelData.Version;
 		}
 	} // CPU::Private
 
@@ -43,17 +53,24 @@ namespace UE::NNERuntimeIREE
 			FIREECompilerRDGResult::StaticStruct()->SerializeBin(Ar, &CompilerResult);
 		}
 
-		FNNERuntimeIREEModelDataHeaderRDG FModelData::ReadHeader(FArchive& Ar)
+		bool FModelData::Load(TConstArrayView64<uint8> Data)
 		{
-			FNNERuntimeIREEModelDataHeaderRDG Header{};
-			FNNERuntimeIREEModelDataHeaderRDG::StaticStruct()->SerializeBin(Ar, &Header);
-			return Header;
+			FMemoryReaderView Reader(Data, /*bIsPersitent =*/ true);
+			Serialize(Reader);
+			return !Reader.IsError();
+		}
+
+		bool FModelData::Store(TArray64<uint8>& Data)
+		{
+			FMemoryWriter64 Writer(Data, /*bIsPersitent =*/ true);
+			Serialize(Writer);
+			return !Writer.IsError();
 		}
 
-		bool FModelData::IsSameGuidAndVersion(FArchive& Ar, FGuid Guid, int32 Version)
+		bool FModelData::ReadHeaderAndIsSameGuidAndVersion(FArchive& Ar, FGuid Guid, int32 Version, FNNERuntimeIREEModelDataHeaderRDG& OutHeader)
 		{
-			FNNERuntimeIREEModelDataHeaderRDG Header = ReadHeader(Ar);
-			return Guid == Header.GUID && Version == Header.Version && !Ar.IsError();
+			FNNERuntimeIREEModelDataHeaderRDG::StaticStruct()->SerializeBin(Ar, &OutHeader);
+			return !Ar.IsError() && Guid == OutHeader.GUID && Version == OutHeader.Version;
 		}
 	} // RDG::Private
-} //UE::NNERuntimeIREE
+} // UE::NNERuntimeIREE
diff --git a/Engine/Plugins/Experimental/NNERuntimeIREE/Source/NNERuntimeIREE/Private/NNERuntimeIREEModelData.h b/Engine/Plugins/Experimental/NNERuntimeIREE/Source/NNERuntimeIREE/Private/NNERuntimeIREEModelData.h
--- a/Engine/Plugins/Experimental/NNERuntimeIREE/Source/NNERuntimeIREE/Private/NNERuntimeIREEModelData.h
+++ b/Engine/Plugins/Experimental/NNERuntimeIREE/Source/NNERuntimeIREE/Private/NNERuntimeIREEModelData.h
@@ -14,12 +14,24 @@
 namespace UE::NNERuntimeIREE::CPU::Private
 {
 	/**
-	 * IREE model data class.
+	 * IREECpu model data class.
 	 */
 	class FModelData
 	{
 	public:
-		void Serialize(FArchive& Ar);
+		/**
+		 * Loads this model from Data that was produced using the Store method
+		 *
+		 * @return true if loading was successfull, false otherwise, this leaves the object in an undefined state
+		 */
+		bool Load(TConstArrayView64<uint8> Data);
+
+		/**
+		* Serializes this model into Data
+		*
+		* @return true if storing was successfull, false otherwise
+		*/
+		bool Store(TArray64<uint8>& Data);
 
 		/**
 		 * Check data for matching Guid and Version without deserializing everything.
@@ -50,11 +62,14 @@ namespace UE::NNERuntimeIREE::CPU::Private
 		 * Compiler output.
 		 */
 		FNNERuntimeIREECompilerResultCPU CompilerResult;
+
+	private:
+		void Serialize(FArchive& Ar);
 	};
 } // UE::NNERuntimeIREE::CPU::Private
 
 /**
- * IREE model data header RDG struct.
+ * IREERdg model data header struct.
  */
 USTRUCT()
 struct FNNERuntimeIREEModelDataHeaderRDG
@@ -89,19 +104,29 @@ struct FNNERuntimeIREEModelDataHeaderRDG
 namespace UE::NNERuntimeIREE::RDG::Private
 {
 	/**
-	 * IREE model data RDG class.
+	 * IREERdg model data class.
 	 */
 	class FModelData
 	{
 	public:
-		void Serialize(FArchive& Ar);
-
 		/**
-		 * Check data for matching Guid and Version with only deserializing the header
+		 * Loads this model from Data that was produced using the Store method
+		 *
+		 * @return true if loading was successfull, false otherwise, this leaves the object in an undefined state
 		 */
-		static bool IsSameGuidAndVersion(FArchive& Ar, FGuid Guid, int32 Version);
+		bool Load(TConstArrayView64<uint8> Data);
 
-		static FNNERuntimeIREEModelDataHeaderRDG ReadHeader(FArchive& Ar);
+		/**
+		* Serializes this model into Data
+		*
+		* @return true if storing was successfull, false otherwise
+		*/
+		bool Store(TArray64<uint8>& Data);
+
+		/**
+		 * Gets Header without loading entire ModelData and checks for matching Guid and Version
+		 */
+		static bool ReadHeaderAndIsSameGuidAndVersion(FArchive& Ar, FGuid Guid, int32 Version, FNNERuntimeIREEModelDataHeaderRDG& OutHeader);
 
 		/**
 		 * Model data header.
@@ -117,5 +142,8 @@ namespace UE::NNERuntimeIREE::RDG::Private
 		 * Compiler output.
 		 */
 		FIREECompilerRDGResult CompilerResult;
+
+	private:
+		void Serialize(FArchive& Ar);
 	};
 } // UE::NNERuntimeIREE::RDG::Private
