[NNE] NNERuntimeIREECpu On Windows IREE compiler should be independent of any developer tools/SDKs.

Problem:
When linking IREE compiler outputed object files on Windows to a shared library (DLL), the clang++ driver defaults to MSVC linker and libraries.

Solution:
Provide the libraries and tools to link without MSVC:
- Use the LLVM linker directly (lld-link.exe)
- Use LLVM tools objdump.exe and lib.exe to create import library of the Windows UCRT DLL, which is packed with the Engine
- Use the UCRT import library for linking to the C Runtime
- Use the Clang Compiler runtime static library to resolve target-specific symbols
- Use manually created libfltused.lib to provide Windows/MSVC specific symbol _fltused

Updated the IREE build scripts to provide the additionally required libraries and tools. Also updated the IREE compiler configs for Windows.

#tests Win64, Editor, NNEAutomationTests all green

#rnx

#Synced-CL 49800879

diff --git a/Engine/Plugins/Experimental/NNERuntimeIREE/Config/IREE_Windows_To_Windows.json b/Engine/Plugins/Experimental/NNERuntimeIREE/Config/IREE_Windows_To_Windows.json
--- a/Engine/Plugins/Experimental/NNERuntimeIREE/Config/IREE_Windows_To_Windows.json
+++ b/Engine/Plugins/Experimental/NNERuntimeIREE/Config/IREE_Windows_To_Windows.json
@@ -2,7 +2,7 @@
 	"ImporterCommand": [ "${PLUGIN_DIR}/Binaries/ThirdParty/IREE/Windows/torch-mlir-import-onnx.exe" ],
 	"ImporterArguments": "-o ${OUTPUT_PATH} --opset-version 21 ${INPUT_PATH}",
 	"CompilerCommand": [ "${PLUGIN_DIR}/Binaries/ThirdParty/IREE/Windows/iree-compile.exe" ],
-	"LinkerCommand": [ "${PLUGIN_DIR}/Binaries/ThirdParty/IREE/Windows/clang++.exe" ],
+	"LinkerCommand": [ "${PLUGIN_DIR}/Binaries/ThirdParty/IREE/Windows/lld-link.exe" ],
 	"SharedLibExt": ".dll",
 	"Targets": [
 		{
@@ -10,27 +10,27 @@
 			"Architecture": "",
 			"x86Features": [ "avx512", "fma", "f16c" ],
 			"CompilerArguments": "--iree-llvmcpu-target-triple=x86_64-pc-windows-msvc --iree-input-type=auto --iree-hal-target-backends=llvm-cpu --iree-llvmcpu-target-cpu-features=+sse,+sse2,+ssse3,+sse4.1,+sse4.2,+avx,+f16c,+avx2,+fma,+avx512f,+avx512vl,+avx512bw,+avx512dq --iree-llvmcpu-link-embedded=false --iree-llvmcpu-link-static --output-format=vm-bytecode --iree-llvmcpu-static-library-output-path=${OBJECT_PATH} -o ${VMFB_PATH} ${INPUT_PATH}",
-			"LinkerArguments": "-shared -o ${SHARED_LIB_PATH} ${OBJECT_PATH}"
+			"LinkerArguments": "-out:${SHARED_LIB_PATH} -nodefaultlib -libpath:${PLUGIN_DIR}/Binaries/ThirdParty/IREE/Windows clang_rt.builtins-x86_64.lib ucrtbase_applocal.lib libfltused.lib -noentry -nologo -dll -alternatename:__chkstk=___chkstk_ms ${OBJECT_PATH}"
 		},
 		{
 			"Name": "AVX2",
 			"Architecture": "",
 			"x86Features": [ "avx2", "fma", "f16c" ],
 			"CompilerArguments": "--iree-llvmcpu-target-triple=x86_64-pc-windows-msvc --iree-input-type=auto --iree-hal-target-backends=llvm-cpu --iree-llvmcpu-target-cpu-features=+sse,+sse2,+ssse3,+sse4.1,+sse4.2,+avx,+f16c,+avx2,+fma --iree-llvmcpu-link-embedded=false --iree-llvmcpu-link-static --output-format=vm-bytecode --iree-llvmcpu-static-library-output-path=${OBJECT_PATH} -o ${VMFB_PATH} ${INPUT_PATH}",
-			"LinkerArguments": "-shared -o ${SHARED_LIB_PATH} ${OBJECT_PATH}"
+			"LinkerArguments": "-out:${SHARED_LIB_PATH} -nodefaultlib -libpath:${PLUGIN_DIR}/Binaries/ThirdParty/IREE/Windows clang_rt.builtins-x86_64.lib ucrtbase_applocal.lib libfltused.lib -noentry -nologo -dll -alternatename:__chkstk=___chkstk_ms ${OBJECT_PATH}"
 		},
 		{
 			"Name": "AVX",
 			"Architecture": "",
 			"x86Features": [ "avx", "f16c" ],
 			"CompilerArguments": "--iree-llvmcpu-target-triple=x86_64-pc-windows-msvc --iree-input-type=auto --iree-hal-target-backends=llvm-cpu --iree-llvmcpu-target-cpu-features=+sse,+sse2,+ssse3,+sse4.1,+sse4.2,+avx,+f16c --iree-llvmcpu-link-embedded=false --iree-llvmcpu-link-static --output-format=vm-bytecode --iree-llvmcpu-static-library-output-path=${OBJECT_PATH} -o ${VMFB_PATH} ${INPUT_PATH}",
-			"LinkerArguments": "-shared -o ${SHARED_LIB_PATH} ${OBJECT_PATH}"
+			"LinkerArguments": "-out:${SHARED_LIB_PATH} -nodefaultlib -libpath:${PLUGIN_DIR}/Binaries/ThirdParty/IREE/Windows clang_rt.builtins-x86_64.lib ucrtbase_applocal.lib libfltused.lib -noentry -nologo -dll -alternatename:__chkstk=___chkstk_ms ${OBJECT_PATH}"
 		},
 		{
 			"Name": "Default",
 			"Architecture": "",
 			"CompilerArguments": "--iree-llvmcpu-target-triple=x86_64-pc-windows-msvc --iree-input-type=auto --iree-hal-target-backends=llvm-cpu --iree-llvmcpu-target-cpu-features=+sse,+sse2,+ssse3,+sse4.1,+sse4.2 --iree-llvmcpu-link-embedded=false --iree-llvmcpu-link-static --output-format=vm-bytecode --iree-llvmcpu-static-library-output-path=${OBJECT_PATH} -o ${VMFB_PATH} ${INPUT_PATH}",
-			"LinkerArguments": "-shared -o ${SHARED_LIB_PATH} ${OBJECT_PATH}"
+			"LinkerArguments": "-out:${SHARED_LIB_PATH} -nodefaultlib -libpath:${PLUGIN_DIR}/Binaries/ThirdParty/IREE/Windows clang_rt.builtins-x86_64.lib ucrtbase_applocal.lib libfltused.lib -noentry -nologo -dll -alternatename:__chkstk=___chkstk_ms ${OBJECT_PATH}"
 		}
 	]
 }
\ No newline at end of file
diff --git a/Engine/Plugins/Experimental/NNERuntimeIREE/Source/ExternalBuild/Scripts/Windows/BuildWindows.bat b/Engine/Plugins/Experimental/NNERuntimeIREE/Source/ExternalBuild/Scripts/Windows/BuildWindows.bat
--- a/Engine/Plugins/Experimental/NNERuntimeIREE/Source/ExternalBuild/Scripts/Windows/BuildWindows.bat
+++ b/Engine/Plugins/Experimental/NNERuntimeIREE/Source/ExternalBuild/Scripts/Windows/BuildWindows.bat
@@ -115,6 +115,17 @@ IF NOT EXIST "%BUILD_DIR%\iree" (
 	ECHO Done.
 
 
+	:: Patch llvm-project
+	ECHO Apply git patch to llvm-project
+	CD "%BUILD_DIR%\iree\third_party\llvm-project"
+	
+	git apply "%UE_PATCHES_DIR%\Windows\llvm-project_compiler-rt_for_msvc_fix.patch"
+	IF ERRORLEVEL 1 GOTO :fail
+	
+	CD ..\..\..\..
+	ECHO Done.
+
+
 	:: Patch SPIRV-Cross
 	ECHO Check for "%SPIRV_CROSS_DIR%"
 	IF EXIST "%SPIRV_CROSS_DIR%" (
@@ -145,6 +156,7 @@ IF NOT EXIST "%BUILD_DIR%\iree-compiler" (
 	SETLOCAL
 	CALL "%VCVARSALL_PATH%" x64 %WIN_SDK_FOR_IREE% -vcvars_ver=%VC_TOOLSET%
 	cmake -G Ninja -B "%BUILD_DIR%\iree-compiler" "%BUILD_DIR%\iree" ^
+		-DLLVM_ENABLE_RUNTIMES=compiler-rt ^
 		-DCMAKE_BUILD_TYPE=%BUILD_TYPE% ^
 		-DIREE_CMAKE_PLUGIN_PATHS="%IREE_CMAKE_PLUGIN_PATH%" ^
 		-DIREE_ENABLE_CPUINFO=OFF ^
@@ -167,6 +179,10 @@ IF NOT EXIST "%BUILD_DIR%\iree-compiler" (
 	cmake --build "%BUILD_DIR%\iree-compiler"
 	IF ERRORLEVEL 1 (ENDLOCAL & GOTO :fail )
 
+	:: Build builtins target separately (clang_rt.builtins-x86_64.lib)
+	cmake --build "%BUILD_DIR%\iree-compiler" --target builtins
+	IF ERRORLEVEL 1 (ENDLOCAL & GOTO :fail )
+
 	ENDLOCAL
 
 	ECHO:
@@ -182,6 +198,10 @@ IF NOT EXIST "%BUILD_DIR%\iree-compiler" (
 
 	cmake --build "%BUILD_DIR%\iree-compiler"
 	IF ERRORLEVEL 1 (ENDLOCAL & GOTO :fail )
+
+	:: Build builtins target separately (clang_rt.builtins-x86_64.lib)
+	cmake --build "%BUILD_DIR%\iree-compiler" --target builtins
+	IF ERRORLEVEL 1 (ENDLOCAL & GOTO :fail )
 	
 	ENDLOCAL
 	ECHO:
diff --git a/Engine/Plugins/Experimental/NNERuntimeIREE/Source/ExternalBuild/Scripts/Windows/InstallWindows.bat b/Engine/Plugins/Experimental/NNERuntimeIREE/Source/ExternalBuild/Scripts/Windows/InstallWindows.bat
--- a/Engine/Plugins/Experimental/NNERuntimeIREE/Source/ExternalBuild/Scripts/Windows/InstallWindows.bat
+++ b/Engine/Plugins/Experimental/NNERuntimeIREE/Source/ExternalBuild/Scripts/Windows/InstallWindows.bat
@@ -74,7 +74,10 @@ ECHO =========================================
 SET "TARGET_IREE_FOLDER=%UE_PLUGIN_ROOT%\Source\ThirdParty\IREE"
 SET "BINARIES_FOLDER=%UE_PLUGIN_ROOT%\Binaries\ThirdParty\IREE\Windows"
 
-CALL :CP "%BUILD_DIR%\iree-compiler\llvm-project\bin" "%BINARIES_FOLDER%" clang++.exe ld.lld.exe
+CALL :CP "%BUILD_DIR%\iree-compiler\llvm-project\bin" "%BINARIES_FOLDER%" lld-link.exe llvm-lib.exe llvm-objdump.exe
+IF ERRORLEVEL 1 GOTO :fail
+
+CALL :CP "%BUILD_DIR%\iree-compiler\llvm-project\lib\clang\21\lib\windows" "%BINARIES_FOLDER%" clang_rt.builtins-x86_64.lib
 IF ERRORLEVEL 1 GOTO :fail
 
 CALL :CP "%BUILD_DIR%\iree-compiler\tools" "%BINARIES_FOLDER%" iree-compile.exe IREECompiler.dll
diff --git a/Engine/Plugins/Experimental/NNERuntimeIREE/Source/NNERuntimeIREE/NNERuntimeIREE.Build.cs b/Engine/Plugins/Experimental/NNERuntimeIREE/Source/NNERuntimeIREE/NNERuntimeIREE.Build.cs
--- a/Engine/Plugins/Experimental/NNERuntimeIREE/Source/NNERuntimeIREE/NNERuntimeIREE.Build.cs
+++ b/Engine/Plugins/Experimental/NNERuntimeIREE/Source/NNERuntimeIREE/NNERuntimeIREE.Build.cs
@@ -1,10 +1,23 @@
 // Copyright Epic Games, Inc. All Rights Reserved.
 
+using EpicGames.Core;
+using Microsoft.Extensions.Logging;
+using System;
+using System.Collections.Generic;
+using System.Diagnostics;
 using System.IO;
+using System.Security.Cryptography;
+using System.Text.RegularExpressions;
 using UnrealBuildTool;
 
 public class NNERuntimeIREE : ModuleRules
 {
+	private static readonly Regex LLVMExportPattern = new(
+		@"^\s*(?:\[\s*)?\d+(?:\s*\])?\s+(?:0x)?[A-Fa-f0-9]+\s+(?<func>\S+)",
+		RegexOptions.Multiline | RegexOptions.Compiled
+	);
+
+
 	public NNERuntimeIREE(ReadOnlyTargetRules Target) : base(Target)
 	{
 		PCHUsage = PCHUsageMode.UseExplicitOrSharedPCHs;
@@ -44,6 +57,27 @@ public class NNERuntimeIREE : ModuleRules
 		{
 			PublicDefinitions.Add("WITH_NNE_RUNTIME_IREE");
 			PublicDefinitions.Add("WITH_NNE_RUNTIME_IREE_RDG");
+
+			if (Target.Type == TargetType.Editor)
+			{
+				// Ensure UCRT import library is present and up-to-date
+				// It is generated from the engine's bundled urctbase.dll (Engine/Binaries/ThirdParty/AppLocalDependencies/...)
+				// The library is not used for linking with UBT; its staged for the NNERuntimeIREECpu compilation.
+				string DllName = "ucrtbase.dll";
+				string DllPath = Path.Combine(EngineDirectory, "Binaries", "ThirdParty", "AppLocalDependencies", "Win64", "x64", "UCRT", DllName);
+
+				string ImportLibraryName = "ucrtbase_applocal.lib";
+				string ImportLibraryPath = Path.Combine(PluginDirectory, "Binaries", "ThirdParty", "IREE", "Windows", ImportLibraryName);
+
+				if (TryEnsureImportLibraryExists(DllPath, ImportLibraryPath))
+				{
+					RuntimeDependencies.Add(ImportLibraryPath);
+				}
+				else
+				{
+					Logger.LogWarning("NNERuntimeIREE: Failed to generate UCRT import library. Runtime IREE CPU compilation may fail.");
+				}
+			}
 		}
 		else if (Target.Platform == UnrealTargetPlatform.Linux)
 		{
@@ -58,4 +92,174 @@ public class NNERuntimeIREE : ModuleRules
 			PublicDefinitions.Add("WITH_NNE_RUNTIME_IREE_RDG");
 		}
 	}
+
+	private bool TryEnsureImportLibraryExists(string DllPath, string LibPath)
+	{
+		if (!File.Exists(DllPath))
+		{
+			Logger.LogWarning($"NNERuntimeIREE: Source DLL is missing: {DllPath}");
+			return false;
+		}
+
+		try
+		{
+			string HashPath = Path.Combine(Path.GetDirectoryName(LibPath), Path.GetFileName(DllPath) + ".hash");
+
+			// Make sure that we run when these change or get deleted
+			ExternalDependencies.Add(DllPath);
+			ExternalDependencies.Add(LibPath);
+			ExternalDependencies.Add(HashPath);
+
+			string CurrentHash = GetFileHash(DllPath);
+
+			if (File.Exists(LibPath) && File.Exists(HashPath))
+			{
+				string StoredHash = File.ReadAllText(HashPath).Trim();
+				if (CurrentHash == StoredHash)
+				{
+					return true;
+				}
+			}
+
+			Stopwatch Timer = Stopwatch.StartNew();
+
+			bool bSuccess = GenerateImportLibrary(DllPath, LibPath);
+
+			Timer.Stop();
+			Logger.LogInformation($"NNERuntimeIREE: Generate import library for {Path.GetFileName(DllPath)} took {Timer.ElapsedMilliseconds} ms.");
+
+			if (bSuccess)
+			{
+				File.WriteAllText(HashPath, CurrentHash);
+
+				return true;
+			}
+
+			return false;
+		}
+		catch (Exception Ex)
+		{
+			Logger.LogWarning($"NNERuntimeIREE: Exception while ensuring import library exists: {Ex.Message}");
+			return false;
+		}
+	}
+
+	private static string GetFileHash(string Filename)
+	{
+		using (var Md5 = MD5.Create())
+		using (var Stream = File.OpenRead(Filename))
+		{
+			return BitConverter.ToString(Md5.ComputeHash(Stream)).Replace("-", "").ToLowerInvariant();
+		}
+	}
+
+	private bool GenerateImportLibrary(string DllPath, string OutputLibPath)
+	{
+		string LLVMBinPath = Path.Combine(PluginDirectory, "Binaries", "ThirdParty", "IREE", "Windows");
+		string ObjDumpExe = Path.Combine(LLVMBinPath, "llvm-objdump.exe");
+		string LibExe = Path.Combine(LLVMBinPath, "llvm-lib.exe");
+
+		if (!File.Exists(ObjDumpExe) || !File.Exists(LibExe))
+		{
+			Logger.LogWarning($"NNERuntimeIREE: Tools missing at path: {LLVMBinPath}");
+			return false;
+		}
+
+		try
+		{
+			// Generate .def
+			ProcessStartInfo DumpInfo = new ProcessStartInfo
+			{
+				FileName = ObjDumpExe,
+				Arguments = $"-p \"{DllPath}\"",
+				RedirectStandardOutput = true,
+				UseShellExecute = false,
+				CreateNoWindow = true
+			};
+
+			List<string> Exports = new List<string>();
+			using (Process Proc = Process.Start(DumpInfo))
+			{
+				string Output = Proc.StandardOutput.ReadToEnd();
+				Proc.WaitForExit();
+
+				if (Proc.ExitCode != 0)
+				{
+					Logger.LogWarning($"NNERuntimeIREE: llvm-objdump.exe failed with exit code: {Proc.ExitCode}");
+					return false;
+				}
+
+				foreach (Match M in LLVMExportPattern.Matches(Output))
+				{
+					string FuncName = M.Groups["func"].Value;
+					if (!string.IsNullOrWhiteSpace(FuncName))
+					{
+						Exports.Add(FuncName);
+					}
+				}
+			}
+
+			if (Exports.Count == 0)
+			{
+				Logger.LogWarning("NNERuntimeIREE: No exports found in DLL.");
+				return false;
+			}
+
+			// Create import library
+			string TempDefPath = Path.GetTempFileName();
+
+			try
+			{
+				List<string> DefContent =
+				[
+					$"LIBRARY \"{Path.GetFileName(DllPath)}\"",
+					"EXPORTS",
+					..Exports,
+				];
+				File.WriteAllLines(TempDefPath, DefContent);
+
+				Directory.CreateDirectory(Path.GetDirectoryName(OutputLibPath));
+
+				ProcessStartInfo LibInfo = new ProcessStartInfo
+				{
+					FileName = LibExe,
+					Arguments = $"/nologo /def:\"{TempDefPath}\" /out:\"{OutputLibPath}\" /machine:x64",
+					RedirectStandardOutput = true,
+					UseShellExecute = false,
+					CreateNoWindow = true
+				};
+				using (Process Proc = Process.Start(LibInfo))
+				{
+					string Output = Proc.StandardOutput.ReadToEnd();
+					Proc.WaitForExit();
+
+					if (Proc.ExitCode != 0)
+					{
+						Logger.LogWarning($"NNERuntimeIREE: llvm-lib.exe failed: {Output}");
+						return false;
+					}
+				}
+			}
+			finally
+			{
+				if (File.Exists(TempDefPath))
+				{
+					File.Delete(TempDefPath);
+				}
+			}
+
+			string ExpPath = Path.ChangeExtension(OutputLibPath, ".exp");
+			if (File.Exists(ExpPath))
+			{
+				File.Delete(ExpPath);
+			}
+
+			return true;
+		}
+		catch (Exception Ex)
+		{
+			Logger.LogWarning($"NNERuntimeIREE: Failed to generate import library: {Ex.Message}");
+			return false;
+		}
+	}
 }
diff --git a/Engine/Plugins/Experimental/NNERuntimeIREE/Source/NNERuntimeIREE/Private/NNERuntimeIREECompiler.cpp b/Engine/Plugins/Experimental/NNERuntimeIREE/Source/NNERuntimeIREE/Private/NNERuntimeIREECompiler.cpp
--- a/Engine/Plugins/Experimental/NNERuntimeIREE/Source/NNERuntimeIREE/Private/NNERuntimeIREECompiler.cpp
+++ b/Engine/Plugins/Experimental/NNERuntimeIREE/Source/NNERuntimeIREE/Private/NNERuntimeIREECompiler.cpp
@@ -256,6 +256,8 @@ namespace UE::NNERuntimeIREE
 				CompilerArguments.ReplaceInline(*FString("${OBJECT_PATH}"), *(FString("\"") + ObjectFilePath + "\""));
 				CompilerArguments.ReplaceInline(*FString("${VMFB_PATH}"), *(FString("\"") + VmfbFilePath + "\""));
 				CompilerArguments.ReplaceInline(*FString("${INPUT_PATH}"), *(FString("\"") + InputFilePath + "\""));
+				CompilerArguments.ReplaceInline(*FString("${PLUGIN_DIR}"), *PluginDir);
+				CompilerArguments.ReplaceInline(*FString("${PROJECT_DIR}"), *FPaths::ProjectDir());
 
 				{
 					SCOPED_NAMED_EVENT_TEXT("Compile", FColor::Magenta);
@@ -286,6 +288,8 @@ namespace UE::NNERuntimeIREE
 				}
 				LinkerArguments.ReplaceInline(*FString("${OBJECT_PATH}"), *(FString("\"") + ObjectFilePath + "\""));
 				LinkerArguments.ReplaceInline(*FString("${SHARED_LIB_PATH}"), *(FString("\"") + SharedLibFilePath + "\""));
+				LinkerArguments.ReplaceInline(*FString("${PLUGIN_DIR}"), *PluginDir);
+				LinkerArguments.ReplaceInline(*FString("${PROJECT_DIR}"), *FPaths::ProjectDir());
 
 				{
 					SCOPED_NAMED_EVENT_TEXT("Link", FColor::Magenta);
