[NNE] NNERuntimeIREE change FNNERuntimeIREEModelDataHeaderRDG to manual serialization

#rnx
#tests NNEModelTest, manual verification that GUID and Version are Serialized the same

In a Future CL we want to add a none USTRUCT field to FNNERuntimeIREEModelDataHeaderRDG. In order to support this FNNERuntimeIREEModelDataHeaderRDG is changed to a none USTRUCT struct and uses manual serialization.

#Synced-CL 49801766

diff --git a/Engine/Plugins/Experimental/NNERuntimeIREE/Source/NNERuntimeIREE/Private/NNERuntimeIREE.cpp b/Engine/Plugins/Experimental/NNERuntimeIREE/Source/NNERuntimeIREE/Private/NNERuntimeIREE.cpp
--- a/Engine/Plugins/Experimental/NNERuntimeIREE/Source/NNERuntimeIREE/Private/NNERuntimeIREE.cpp
+++ b/Engine/Plugins/Experimental/NNERuntimeIREE/Source/NNERuntimeIREE/Private/NNERuntimeIREE.cpp
@@ -603,7 +603,7 @@ int32 UNNERuntimeIREEVulkan::GetVersion() const
 #ifdef WITH_NNE_RUNTIME_IREE_RDG
 
 FGuid UNNERuntimeIREERdg::GUID = FGuid((int32)'I', (int32)'R', (int32)'D', (int32)'G');
-int32 UNNERuntimeIREERdg::Version = 0x00000003; // Increment each time iree compiler/runtime is updated and each time FModelData serialization changes
+int32 UNNERuntimeIREERdg::Version = 0x00000004; // Increment each time iree compiler/runtime is updated and each time FModelData serialization changes
 
 FString UNNERuntimeIREERdg::GetRuntimeName() const
 {
@@ -659,7 +659,7 @@ TSharedPtr<UE::NNE::FSharedModelData> UNNERuntimeIREERdg::CreateModelData(const
 	{
 		SCOPED_NAMED_EVENT_TEXT("Validate", FColor::Magenta);
 
-		FNNERuntimeIREEModelDataHeaderRDG Header{};
+		FModelDataHeader Header{};
 		TUniquePtr<FArchive> FileReader(IFileManager::Get().CreateFileReader(*IREEModelDataFilePath));
 		if (FileReader && FModelData::ReadHeaderAndIsSameGuidAndVersion(*FileReader, GUID, Version, Header))
 		{
@@ -772,7 +772,7 @@ UNNERuntimeIREERdg::ECanCreateModelRDGStatus UNNERuntimeIREERdg::CanCreateModelR
 		return ECanCreateModelRDGStatus::Fail;
 	}
 
-	FNNERuntimeIREEModelDataHeaderRDG Header{};
+	FModelDataHeader Header{};
 	FMemoryReaderView Reader(SharedData->GetView(), /*bIsPersitent =*/ true);
 	if (!FModelData::ReadHeaderAndIsSameGuidAndVersion(Reader, GUID, Version, Header))
 	{
diff --git a/Engine/Plugins/Experimental/NNERuntimeIREE/Source/NNERuntimeIREE/Private/NNERuntimeIREEModelData.cpp b/Engine/Plugins/Experimental/NNERuntimeIREE/Source/NNERuntimeIREE/Private/NNERuntimeIREEModelData.cpp
--- a/Engine/Plugins/Experimental/NNERuntimeIREE/Source/NNERuntimeIREE/Private/NNERuntimeIREEModelData.cpp
+++ b/Engine/Plugins/Experimental/NNERuntimeIREE/Source/NNERuntimeIREE/Private/NNERuntimeIREEModelData.cpp
@@ -46,9 +46,17 @@ namespace UE::NNERuntimeIREE
 
 	namespace RDG::Private
 	{
+		void FModelDataHeader::Serialize(FArchive& Ar)
+		{
+			Ar << GUID;
+			Ar << Version;
+			Ar << FileId;
+			Ar << ShaderPlatforms;
+		}
+
 		void FModelData::Serialize(FArchive& Ar)
 		{
-			FNNERuntimeIREEModelDataHeaderRDG::StaticStruct()->SerializeBin(Ar, &Header);
+			Header.Serialize(Ar);
 			ModuleMetaData.Serialize(Ar);
 			FIREECompilerRDGResult::StaticStruct()->SerializeBin(Ar, &CompilerResult);
 		}
@@ -67,9 +75,9 @@ namespace UE::NNERuntimeIREE
 			return !Writer.IsError();
 		}
 
-		bool FModelData::ReadHeaderAndIsSameGuidAndVersion(FArchive& Ar, FGuid Guid, int32 Version, FNNERuntimeIREEModelDataHeaderRDG& OutHeader)
+		bool FModelData::ReadHeaderAndIsSameGuidAndVersion(FArchive& Ar, FGuid Guid, int32 Version, FModelDataHeader& OutHeader)
 		{
-			FNNERuntimeIREEModelDataHeaderRDG::StaticStruct()->SerializeBin(Ar, &OutHeader);
+			OutHeader.Serialize(Ar);
 			return !Ar.IsError() && Guid == OutHeader.GUID && Version == OutHeader.Version;
 		}
 	} // RDG::Private
diff --git a/Engine/Plugins/Experimental/NNERuntimeIREE/Source/NNERuntimeIREE/Private/NNERuntimeIREEModelData.h b/Engine/Plugins/Experimental/NNERuntimeIREE/Source/NNERuntimeIREE/Private/NNERuntimeIREEModelData.h
--- a/Engine/Plugins/Experimental/NNERuntimeIREE/Source/NNERuntimeIREE/Private/NNERuntimeIREEModelData.h
+++ b/Engine/Plugins/Experimental/NNERuntimeIREE/Source/NNERuntimeIREE/Private/NNERuntimeIREEModelData.h
@@ -7,9 +7,6 @@
 #include "NNERuntimeIREECompiler.h"
 #include "NNERuntimeIREEMetaData.h"
 #include "NNETypes.h"
-#include "UObject/Class.h"
-
-#include "NNERuntimeIREEModelData.generated.h"
 
 namespace UE::NNERuntimeIREE::CPU::Private
 {
@@ -56,7 +53,7 @@ namespace UE::NNERuntimeIREE::CPU::Private
 		/**
 		 * Module meta data.
 		 */
-		UE::NNERuntimeIREE::FModuleMetaData ModuleMetaData;
+		FModuleMetaData ModuleMetaData;
 
 		/**
 		 * Compiler output.
@@ -68,41 +65,37 @@ namespace UE::NNERuntimeIREE::CPU::Private
 	};
 } // UE::NNERuntimeIREE::CPU::Private
 
-/**
- * IREERdg model data header struct.
- */
-USTRUCT()
-struct FNNERuntimeIREEModelDataHeaderRDG
-{
-	GENERATED_BODY()
 
+namespace UE::NNERuntimeIREE::RDG::Private
+{
 	/**
-	 * A Guid that uniquely identifies this IREE model data.
+	 * IREERdg model data header struct.
 	 */
-	UPROPERTY()
-	FGuid GUID;
+	struct FModelDataHeader
+	{
+		void Serialize(FArchive& Ar);
 
-	/**
-	 * Current version of this IREE model data.
-	 */
-	UPROPERTY()
-	int32 Version = 0;
+		/**
+		 * A Guid that uniquely identifies this IREE model data.
+		 */
+		FGuid GUID;
 
-	/**
-	 * A Guid that uniquely identifies the model.
-	 */
-	UPROPERTY()
-	FGuid FileId;
+		/**
+		 * Current version of this IREE model data.
+		 */
+		int32 Version = 0;
 
-	/**
-	 * List of shader platforms supported by this IREE model data.
-	 */
-	UPROPERTY()
-	TArray<FString> ShaderPlatforms;
-};
+		/**
+		 * A Guid that uniquely identifies the model.
+		 */
+		FGuid FileId;
+
+		/**
+		 * List of shader platforms supported by this IREE model data.
+		 */
+		TArray<FString> ShaderPlatforms;
+	};
 
-namespace UE::NNERuntimeIREE::RDG::Private
-{
 	/**
 	 * IREERdg model data class.
 	 */
@@ -126,17 +119,17 @@ namespace UE::NNERuntimeIREE::RDG::Private
 		/**
 		 * Gets Header without loading entire ModelData and checks for matching Guid and Version
 		 */
-		static bool ReadHeaderAndIsSameGuidAndVersion(FArchive& Ar, FGuid Guid, int32 Version, FNNERuntimeIREEModelDataHeaderRDG& OutHeader);
+		static bool ReadHeaderAndIsSameGuidAndVersion(FArchive& Ar, FGuid Guid, int32 Version, FModelDataHeader& OutHeader);
 
 		/**
 		 * Model data header.
 		 */
-		FNNERuntimeIREEModelDataHeaderRDG Header;
+		FModelDataHeader Header;
 
 		/**
 		 * Module meta data.
 		 */
-		UE::NNERuntimeIREE::FModuleMetaData ModuleMetaData;
+		FModuleMetaData ModuleMetaData;
 
 		/**
 		 * Compiler output.
