[NNE] NNERuntimeIREE compiler cache invalidation

#rnx
#tests NNEModelTest

Compute the hash of the contents of the .json file and the modification time and file size of the importer, compiler and linker files. Store this hash inside IREEModelData and use it to figure out if recompilation is required during CreateModelData.

#Synced-CL 49944062

diff --git a/Engine/Plugins/Experimental/NNERuntimeIREE/Source/IREEDriverRDG/Internal/IREECompilerRDG.h b/Engine/Plugins/Experimental/NNERuntimeIREE/Source/IREEDriverRDG/Internal/IREECompilerRDG.h
--- a/Engine/Plugins/Experimental/NNERuntimeIREE/Source/IREEDriverRDG/Internal/IREECompilerRDG.h
+++ b/Engine/Plugins/Experimental/NNERuntimeIREE/Source/IREEDriverRDG/Internal/IREECompilerRDG.h
@@ -3,6 +3,7 @@
 #pragma once
 
 #include "CoreMinimal.h"
+#include "Misc/SecureHash.h"
 #include "RHIShaderPlatform.h"
 #include "Serialization/JsonSerializerMacros.h"
 
@@ -110,13 +111,14 @@ namespace UE::IREE::Compiler::RDG
 	class IREEDRIVERRDG_API FCompiler
 	{
 	private:
-		FCompiler(const ITargetPlatform* TargetPlatform, const FString& InImporterCommand, const FString& InImporterArguments, const FString& InCompilerCommand, const FString& InSharedLibExt, TConstArrayView<FBuildTarget> InBuildTargets);
+		FCompiler(const ITargetPlatform* TargetPlatform, const FString& InImporterCommand, const FString& InImporterArguments, const FString& InCompilerCommand, const FString& InSharedLibExt, TConstArrayView<FBuildTarget> InBuildTargets, const FMD5Hash& InHash);
 
 	public:
 		~FCompiler() {};
 		static TUniquePtr<FCompiler> Make(const ITargetPlatform* TargetPlatform);
 
 	public:
+		const FMD5Hash& GetHash() const { return Hash; }
 		bool ImportOnnx(TConstArrayView<uint8> InFileData, const FString& InModelName, const FString& InOutputDir, TArray64<uint8>& OutMlirData);
 		bool CompileMlir(TConstArrayView<uint8> InFileData, const FString& InModelName, const FString& InOutputDir, TConstArrayView<EShaderPlatform> ShaderPlatforms, FIREECompilerRDGResult& OutCompilerResult);
 
@@ -126,6 +128,7 @@ namespace UE::IREE::Compiler::RDG
 		FString ImporterArguments;
 		FString CompilerCommand;
 		TArray<FBuildTarget> BuildTargets;
+		const FMD5Hash Hash;
 	};
 } // namespace UE::IREE::Compiler::RDG
 
diff --git a/Engine/Plugins/Experimental/NNERuntimeIREE/Source/IREEDriverRDG/Private/IREECompilerRDG.cpp b/Engine/Plugins/Experimental/NNERuntimeIREE/Source/IREEDriverRDG/Private/IREECompilerRDG.cpp
--- a/Engine/Plugins/Experimental/NNERuntimeIREE/Source/IREEDriverRDG/Private/IREECompilerRDG.cpp
+++ b/Engine/Plugins/Experimental/NNERuntimeIREE/Source/IREEDriverRDG/Private/IREECompilerRDG.cpp
@@ -147,13 +147,15 @@ namespace UE::IREE::Compiler
 			}
 		} // Private
 
-		FCompiler::FCompiler(const ITargetPlatform* InTargetPlatform, const FString& InImporterCommand, const FString& InImporterArguments, const FString& InCompilerCommand, const FString& InSharedLibExt, TConstArrayView<FBuildTarget> InBuildTargets) : TargetPlatform(InTargetPlatform), ImporterCommand(InImporterCommand), ImporterArguments(InImporterArguments), CompilerCommand(InCompilerCommand), BuildTargets(InBuildTargets)
+		FCompiler::FCompiler(const ITargetPlatform* InTargetPlatform, const FString& InImporterCommand, const FString& InImporterArguments, const FString& InCompilerCommand, const FString& InSharedLibExt, TConstArrayView<FBuildTarget> InBuildTargets, const FMD5Hash& InHash) : TargetPlatform(InTargetPlatform), ImporterCommand(InImporterCommand), ImporterArguments(InImporterArguments), CompilerCommand(InCompilerCommand), BuildTargets(InBuildTargets), Hash(InHash)
 		{
 			
 		}
 
 		TUniquePtr<FCompiler> FCompiler::Make(const ITargetPlatform* TargetPlatform)
 		{
+			SCOPED_NAMED_EVENT_TEXT("FCompiler::Make", FColor::Magenta);
+
 			using namespace Private;
 
 			const FString TargetPlatformName = TargetPlatform ? TargetPlatform->IniPlatformName() : UGameplayStatics::GetPlatformName();
@@ -173,6 +175,7 @@ namespace UE::IREE::Compiler
 			FString CompilerCommand;
 			FString SharedLibExt;
 			TArray<FBuildTarget> BuildTargets;
+			FMD5Hash Hash;
 			IPlatformFile& PlatformFile = FPlatformFileManager::Get().GetPlatformFile();
 			for (FString BuildConfigFilePath : BuildConfigFilePaths)
 			{
@@ -244,6 +247,17 @@ namespace UE::IREE::Compiler
 							ImporterArguments = BuildConfig.ImporterArguments;
 							CompilerCommand = TmpCompilerCommand;
 							BuildTargets = BuildConfig.BuildTargets;
+
+							FMD5 HashBuilder{};
+							bool bHashSuccess = true;
+							UE::IREEUtils::HashAppendString(HashBuilder, BuildConfigFileString);
+							bHashSuccess &= UE::IREEUtils::HashAppendFileStat(HashBuilder, ImporterCommand);
+							bHashSuccess &= UE::IREEUtils::HashAppendFileStat(HashBuilder, CompilerCommand);
+							Hash.Set(HashBuilder);
+							if (!bHashSuccess)
+							{
+								UE_LOG(LogIREEDriverRDG, Warning, TEXT("UNNERuntimeIREERDG failed to create a valid compiler hash, while loading %s"), *BuildConfigFilePath);
+							}
 							break;
 						}
 						else
@@ -261,7 +275,7 @@ namespace UE::IREE::Compiler
 			{
 				return TUniquePtr<FCompiler>();
 			}
-			return TUniquePtr<FCompiler>(new FCompiler(TargetPlatform, ImporterCommand, ImporterArguments, CompilerCommand, SharedLibExt, BuildTargets));
+			return TUniquePtr<FCompiler>(new FCompiler(TargetPlatform, ImporterCommand, ImporterArguments, CompilerCommand, SharedLibExt, BuildTargets, Hash));
 		}
 
 		bool FCompiler::ImportOnnx(TConstArrayView<uint8> InFileData, const FString& InModelName, const FString& InOutputDir, TArray64<uint8>& OutMlirData)
diff --git a/Engine/Plugins/Experimental/NNERuntimeIREE/Source/IREEUtils/Internal/IREEUtils.h b/Engine/Plugins/Experimental/NNERuntimeIREE/Source/IREEUtils/Internal/IREEUtils.h
--- a/Engine/Plugins/Experimental/NNERuntimeIREE/Source/IREEUtils/Internal/IREEUtils.h
+++ b/Engine/Plugins/Experimental/NNERuntimeIREE/Source/IREEUtils/Internal/IREEUtils.h
@@ -3,6 +3,7 @@
 #pragma once
 
 #include "Containers/UnrealString.h"
+#include "Misc/SecureHash.h"
 
 namespace UE::IREEUtils
 {
@@ -15,4 +16,23 @@ IREEUTILS_API void RunCommand(const FString& Command, const FString& Arguments,
 
 IREEUTILS_API bool ImportOnnx(const FString& ImporterCommand, const FString& ImporterArguments, TConstArrayView<uint8> InFileData, const FString& InModelName, const FString& InOutputDir, TArray64<uint8>& OutMlirData);
 
+/**
+* Hashes the contents of the entire file. Depending on the file size this might be a slow operation.
+* 
+* @return true on success, false otherwise
+*/
+IREEUTILS_API bool HashAppendFile(FMD5& Hash, const FString& InFilePath);
+
+/**
+* Hashes the modification time and the file size.
+* 
+* @return true on success, false otherwise
+*/
+IREEUTILS_API bool HashAppendFileStat(FMD5& Hash, const FString& InFilePath);
+
+/**
+* Hashes the char array of an FString. This isn't guaranteed to produce the same hash for the same string on different platforms.
+*/
+IREEUTILS_API void HashAppendString(FMD5& Hash, const FString& Data);
+
 }
\ No newline at end of file
diff --git a/Engine/Plugins/Experimental/NNERuntimeIREE/Source/IREEUtils/Private/IREEUtils.cpp b/Engine/Plugins/Experimental/NNERuntimeIREE/Source/IREEUtils/Private/IREEUtils.cpp
--- a/Engine/Plugins/Experimental/NNERuntimeIREE/Source/IREEUtils/Private/IREEUtils.cpp
+++ b/Engine/Plugins/Experimental/NNERuntimeIREE/Source/IREEUtils/Private/IREEUtils.cpp
@@ -11,6 +11,7 @@
 #include "Misc/MonitoredProcess.h"
 #include "Misc/Optional.h"
 #include "Misc/Paths.h"
+#include "Memory/MemoryView.h"
 
 namespace UE::IREEUtils
 {
@@ -208,4 +209,37 @@ bool ImportOnnx(const FString& ImporterCommand, const FString& ImporterArguments
 	return true;
 }
 
+bool HashAppendFile(FMD5& Hash, const FString& InFilePath)
+{
+	SCOPED_NAMED_EVENT_TEXT(TEXT("IREEUtils::HashAppendFile"), FColor::Magenta);
+
+	bool bSuccess = FFileHelper::LoadFileInBlocks(InFilePath,
+		[&Hash](FMemoryView Block)
+		{
+			Hash.Update((uint8*)Block.GetData(), Block.GetSize());
+		});
+	return bSuccess;
+}
+
+bool HashAppendFileStat(FMD5& Hash, const FString& InFilePath)
+{
+	FFileStatData Stat = IFileManager::Get().GetStatData(*InFilePath);
+	if (Stat.bIsValid)
+	{
+		int64 Data[2] =
+		{
+			Stat.ModificationTime.GetTicks(),
+			Stat.FileSize
+		};
+		Hash.Update((uint8*)Data, sizeof(Data));
+		return true;
+	}
+	return false;
+}
+
+void HashAppendString(FMD5& Hash, const FString& Data)
+{
+	Hash.Update((const uint8*)Data.GetCharArray().GetData(), Data.GetCharArray().NumBytes());
+}
+
 } // namespace UE::IREEUtils
\ No newline at end of file
diff --git a/Engine/Plugins/Experimental/NNERuntimeIREE/Source/NNERuntimeIREE/Private/NNERuntimeIREE.cpp b/Engine/Plugins/Experimental/NNERuntimeIREE/Source/NNERuntimeIREE/Private/NNERuntimeIREE.cpp
--- a/Engine/Plugins/Experimental/NNERuntimeIREE/Source/NNERuntimeIREE/Private/NNERuntimeIREE.cpp
+++ b/Engine/Plugins/Experimental/NNERuntimeIREE/Source/NNERuntimeIREE/Private/NNERuntimeIREE.cpp
@@ -210,7 +210,7 @@ namespace UE::NNERuntimeIREE::Private
 } // UE::NNERuntimeIREE::Private
 
 FGuid UNNERuntimeIREECpu::GUID = FGuid((int32)'I', (int32)'C', (int32)'P', (int32)'U');
-int32 UNNERuntimeIREECpu::Version = 0x00000006; // Increment each time iree compiler/runtime is updated and each time FModelData serialization changes
+int32 UNNERuntimeIREECpu::Version = 0x00000007; // Increment each time iree compiler/runtime is updated and each time FModelData serialization changes
 
 void UNNERuntimeIREECpu::Init(TSharedRef<UE::NNERuntimeIREE::Private::FEnvironment> InEnvironment)
 {
@@ -264,6 +264,12 @@ TSharedPtr<UE::NNE::FSharedModelData> UNNERuntimeIREECpu::CreateModelData(const
 
 	FString IREEModelDataFilePath = FPaths::Combine(IntermediateDirFullPath, FileIdString) + ".ireemodeldata";
 
+	TUniquePtr<UE::NNERuntimeIREE::CPU::FCompiler> Compiler = UE::NNERuntimeIREE::CPU::FCompiler::Make(TargetPlatformName);
+	if (!Compiler.IsValid())
+	{
+		UE_LOG(LogNNERuntimeIREE, Warning, TEXT("UNNERuntimeIREECpu failed to create a compiler to compile for platform %s"), *TargetPlatformName);
+		return TSharedPtr<UE::NNE::FSharedModelData>();
+	}
 	TArray64<uint8> ResultData;
 	FModelData IREEModelData{};
 
@@ -276,7 +282,7 @@ TSharedPtr<UE::NNE::FSharedModelData> UNNERuntimeIREECpu::CreateModelData(const
 		{
 			check(FileIdString.Equals(IREEModelData.FileId.ToString(EGuidFormats::Digits).ToLower()));
 
-			bNeedCompileMlir = false;
+			bNeedCompileMlir = Compiler->GetHash() != IREEModelData.CompilerHash;
 			for (int32 i = 0; i < IREEModelData.CompilerResult.ArchitectureInfos.Num(); i++)
 			{
 				const FNNERuntimeIREEArchitectureInfoCPU& Info = IREEModelData.CompilerResult.ArchitectureInfos[i];
@@ -298,13 +304,6 @@ TSharedPtr<UE::NNE::FSharedModelData> UNNERuntimeIREECpu::CreateModelData(const
 		PlatformFile.DeleteDirectoryRecursively(*IntermediateDirFullPath);
 		PlatformFile.CreateDirectoryTree(*IntermediateDirFullPath);
 
-		TUniquePtr<UE::NNERuntimeIREE::CPU::FCompiler> Compiler = UE::NNERuntimeIREE::CPU::FCompiler::Make(TargetPlatformName);
-		if (!Compiler.IsValid())
-		{
-			UE_LOG(LogNNERuntimeIREE, Warning, TEXT("UNNERuntimeIREECpu failed to create a compiler to compile for platform %s"), *TargetPlatformName);
-			return TSharedPtr<UE::NNE::FSharedModelData>();
-		}
-
 		TArray64<uint8> ImporterOutputData;
 		if(FileType.Compare(TEXT("onnx"), ESearchCase::IgnoreCase) == 0)
 		{
@@ -327,6 +326,7 @@ TSharedPtr<UE::NNE::FSharedModelData> UNNERuntimeIREECpu::CreateModelData(const
 		IREEModelData.GUID = UNNERuntimeIREECpu::GUID;
 		IREEModelData.Version = UNNERuntimeIREECpu::Version;
 		IREEModelData.FileId = FileId;
+		IREEModelData.CompilerHash = Compiler->GetHash();
 		if (AdditionalFileData.Contains("IREEModuleMetaData") && !AdditionalFileData["IREEModuleMetaData"].IsEmpty())
 		{
 			FMemoryReaderView Reader(AdditionalFileData["IREEModuleMetaData"], /*bIsPersitent =*/ true);
@@ -603,7 +603,7 @@ int32 UNNERuntimeIREEVulkan::GetVersion() const
 #ifdef WITH_NNE_RUNTIME_IREE_RDG
 
 FGuid UNNERuntimeIREERdg::GUID = FGuid((int32)'I', (int32)'R', (int32)'D', (int32)'G');
-int32 UNNERuntimeIREERdg::Version = 0x00000004; // Increment each time iree compiler/runtime is updated and each time FModelData serialization changes
+int32 UNNERuntimeIREERdg::Version = 0x00000005; // Increment each time iree compiler/runtime is updated and each time FModelData serialization changes
 
 FString UNNERuntimeIREERdg::GetRuntimeName() const
 {
@@ -654,6 +654,13 @@ TSharedPtr<UE::NNE::FSharedModelData> UNNERuntimeIREERdg::CreateModelData(const
 
 	FString IREEModelDataFilePath = FPaths::Combine(IntermediateDirFullPath, FileIdString) + ".ireemodeldata";
 
+	TUniquePtr<UE::IREE::Compiler::RDG::FCompiler> Compiler = UE::IREE::Compiler::RDG::FCompiler::Make(TargetPlatform);
+	if (!Compiler.IsValid())
+	{
+		UE_LOG(LogNNERuntimeIREE, Warning, TEXT("UNNERuntimeIREERdg failed to create a compiler to compile for platform %s"), *TargetPlatformName);
+		return TSharedPtr<UE::NNE::FSharedModelData>();
+	}
+
 	bool bNeedCompileMlir = true;
 	if (PlatformFile.FileExists(*IREEModelDataFilePath))
 	{
@@ -665,7 +672,7 @@ TSharedPtr<UE::NNE::FSharedModelData> UNNERuntimeIREERdg::CreateModelData(const
 		{
 			check(FileIdString.Equals(Header.FileId.ToString(EGuidFormats::Digits).ToLower()));
 
-			bNeedCompileMlir = false;
+			bNeedCompileMlir = Compiler->GetHash() != Header.CompilerHash;
 			for (EShaderPlatform ShaderPlatform : ShaderPlatforms)
 			{
 				const FString ShaderPlatformName = LexToString(ShaderPlatform);
@@ -683,13 +690,6 @@ TSharedPtr<UE::NNE::FSharedModelData> UNNERuntimeIREERdg::CreateModelData(const
 		PlatformFile.DeleteDirectoryRecursively(*IntermediateDirFullPath);
 		PlatformFile.CreateDirectoryTree(*IntermediateDirFullPath);
 
-		TUniquePtr<UE::IREE::Compiler::RDG::FCompiler> Compiler = UE::IREE::Compiler::RDG::FCompiler::Make(TargetPlatform);
-		if (!Compiler.IsValid())
-		{
-			UE_LOG(LogNNERuntimeIREE, Warning, TEXT("UNNERuntimeIREERdg failed to create a compiler to compile for platform %s"), *TargetPlatformName);
-			return TSharedPtr<UE::NNE::FSharedModelData>();
-		}
-
 		TArray64<uint8> ImporterOutputData;
 		if(FileType.Compare(TEXT("onnx"), ESearchCase::IgnoreCase) == 0)
 		{
@@ -713,6 +713,7 @@ TSharedPtr<UE::NNE::FSharedModelData> UNNERuntimeIREERdg::CreateModelData(const
 		IREEModelData.Header.GUID = UNNERuntimeIREERdg::GUID;
 		IREEModelData.Header.Version = UNNERuntimeIREERdg::Version;
 		IREEModelData.Header.FileId = FileId;
+		IREEModelData.Header.CompilerHash = Compiler->GetHash();
 		for (EShaderPlatform ShaderPlatform : ShaderPlatforms)
 		{
 			const FString ShaderPlatformName = LexToString(ShaderPlatform);
diff --git a/Engine/Plugins/Experimental/NNERuntimeIREE/Source/NNERuntimeIREE/Private/NNERuntimeIREECompiler.cpp b/Engine/Plugins/Experimental/NNERuntimeIREE/Source/NNERuntimeIREE/Private/NNERuntimeIREECompiler.cpp
--- a/Engine/Plugins/Experimental/NNERuntimeIREE/Source/NNERuntimeIREE/Private/NNERuntimeIREECompiler.cpp
+++ b/Engine/Plugins/Experimental/NNERuntimeIREE/Source/NNERuntimeIREE/Private/NNERuntimeIREECompiler.cpp
@@ -45,14 +45,16 @@ namespace UE::NNERuntimeIREE
 			}
 		} // Private
 
-		FCompiler::FCompiler(const FString& InImporterCommand, const FString& InImporterArguments, const FString& InCompilerCommand, const FString& InLinkerCommand, const FString& InSharedLibExt, TConstArrayView<FBuildTarget> InBuildTargets) 
-			: ImporterCommand(InImporterCommand), ImporterArguments(InImporterArguments), CompilerCommand(InCompilerCommand), LinkerCommand(InLinkerCommand), SharedLibExt(InSharedLibExt), BuildTargets(InBuildTargets)
+		FCompiler::FCompiler(const FString& InImporterCommand, const FString& InImporterArguments, const FString& InCompilerCommand, const FString& InLinkerCommand, const FString& InSharedLibExt, TConstArrayView<FBuildTarget> InBuildTargets, const FMD5Hash& InHash) 
+			: ImporterCommand(InImporterCommand), ImporterArguments(InImporterArguments), CompilerCommand(InCompilerCommand), LinkerCommand(InLinkerCommand), SharedLibExt(InSharedLibExt), BuildTargets(InBuildTargets), Hash(InHash)
 		{
 
 		}
 
 		TUniquePtr<FCompiler> FCompiler::Make(const FString& InTargetPlatformName)
 		{
+			SCOPED_NAMED_EVENT_TEXT("FCompiler::Make", FColor::Magenta);
+
 			using namespace Private;
 
 			FString PluginDir = FPaths::ConvertRelativePathToFull(*IPluginManager::Get().FindPlugin(UE_PLUGIN_NAME)->GetBaseDir());
@@ -71,6 +73,7 @@ namespace UE::NNERuntimeIREE
 			FString LinkerCommand;
 			FString SharedLibExt;
 			TArray<FBuildTarget> BuildTargets;
+			FMD5Hash Hash;
 			IPlatformFile& PlatformFile = FPlatformFileManager::Get().GetPlatformFile();
 			for (FString BuildConfigFilePath : BuildConfigFilePaths)
 			{
@@ -79,10 +82,8 @@ namespace UE::NNERuntimeIREE
 					FString BuildConfigFileString;
 					if (FFileHelper::LoadFileToString(BuildConfigFileString, *BuildConfigFilePath))
 					{
-						TSharedRef<TJsonReader<TCHAR>> JsonReader = TJsonReaderFactory<TCHAR>::Create(BuildConfigFileString);
-						TSharedPtr<FJsonObject> JsonObject = MakeShareable(new FJsonObject);
 						FBuildConfig BuildConfig;
-						if (FJsonSerializer::Deserialize(JsonReader, JsonObject) && JsonObject.IsValid() && BuildConfig.FromJson(JsonObject))
+						if(BuildConfig.FromJson(BuildConfigFileString))
 						{
 							if (BuildConfig.BuildTargets.IsEmpty())
 							{
@@ -168,6 +169,18 @@ namespace UE::NNERuntimeIREE
 							LinkerCommand = TmpLinkerCommand;
 							SharedLibExt = BuildConfig.SharedLibExt;
 							BuildTargets = BuildConfig.BuildTargets;
+
+							FMD5 HashBuilder{};
+							bool bHashSuccess = true;
+							UE::IREEUtils::HashAppendString(HashBuilder, BuildConfigFileString);
+							bHashSuccess &= UE::IREEUtils::HashAppendFileStat(HashBuilder, ImporterCommand);
+							bHashSuccess &= UE::IREEUtils::HashAppendFileStat(HashBuilder, CompilerCommand);
+							bHashSuccess &= UE::IREEUtils::HashAppendFileStat(HashBuilder, LinkerCommand);
+							Hash.Set(HashBuilder);
+							if (!bHashSuccess)
+							{
+								UE_LOG(LogNNERuntimeIREE, Warning, TEXT("UNNERuntimeIREECpu failed to create a valid compiler hash, while loading %s"), *BuildConfigFilePath);
+							}
 							break;
 						}
 						else
@@ -185,7 +198,7 @@ namespace UE::NNERuntimeIREE
 			{
 				return TUniquePtr<FCompiler>();
 			}
-			return TUniquePtr<FCompiler>(new FCompiler(ImporterCommand, ImporterArguments, CompilerCommand, LinkerCommand, SharedLibExt, BuildTargets));
+			return TUniquePtr<FCompiler>(new FCompiler(ImporterCommand, ImporterArguments, CompilerCommand, LinkerCommand, SharedLibExt, BuildTargets, Hash));
 		}
 
 		bool FCompiler::ImportOnnx(TConstArrayView<uint8> InFileData, const FString& InModelName, const FString& InOutputDir, TArray64<uint8>& OutMlirData)
diff --git a/Engine/Plugins/Experimental/NNERuntimeIREE/Source/NNERuntimeIREE/Private/NNERuntimeIREECompiler.h b/Engine/Plugins/Experimental/NNERuntimeIREE/Source/NNERuntimeIREE/Private/NNERuntimeIREECompiler.h
--- a/Engine/Plugins/Experimental/NNERuntimeIREE/Source/NNERuntimeIREE/Private/NNERuntimeIREECompiler.h
+++ b/Engine/Plugins/Experimental/NNERuntimeIREE/Source/NNERuntimeIREE/Private/NNERuntimeIREECompiler.h
@@ -4,6 +4,7 @@
 
 #include "CoreMinimal.h"
 #include "Serialization/JsonSerializerMacros.h"
+#include "Misc/SecureHash.h"
 
 #include "NNERuntimeIREECompiler.generated.h"
 
@@ -89,13 +90,14 @@ namespace UE::NNERuntimeIREE::CPU
 	class FCompiler
 	{
 	private:
-		FCompiler(const FString& InImporterCommand, const FString& InImporterArguments, const FString& InCompilerCommand, const FString& InLinkerCommand, const FString& InSharedLibExt, TConstArrayView<FBuildTarget> InBuildTargets);
+		FCompiler(const FString& InImporterCommand, const FString& InImporterArguments, const FString& InCompilerCommand, const FString& InLinkerCommand, const FString& InSharedLibExt, TConstArrayView<FBuildTarget> InBuildTargets, const FMD5Hash& InHash);
 
 	public:
 		~FCompiler() {};
 		static TUniquePtr<FCompiler> Make(const FString& InTargetPlatformName);
 
 	public:
+		const FMD5Hash& GetHash() const { return Hash; }
 		bool ImportOnnx(TConstArrayView<uint8> InFileData, const FString& InModelName, const FString& InOutputDir, TArray64<uint8>& OutMlirData);
 		bool CompileMlir(TConstArrayView<uint8> InFileData, const FString& InModelName, const FString& InOutputDir, FNNERuntimeIREECompilerResultCPU& OutCompilerResult);
 
@@ -106,6 +108,7 @@ namespace UE::NNERuntimeIREE::CPU
 		FString LinkerCommand;
 		FString SharedLibExt;
 		TArray<FBuildTarget> BuildTargets;
+		const FMD5Hash Hash;
 	};
 } // UE::NNERuntimeIREE::CPU
 
diff --git a/Engine/Plugins/Experimental/NNERuntimeIREE/Source/NNERuntimeIREE/Private/NNERuntimeIREEModelData.cpp b/Engine/Plugins/Experimental/NNERuntimeIREE/Source/NNERuntimeIREE/Private/NNERuntimeIREEModelData.cpp
--- a/Engine/Plugins/Experimental/NNERuntimeIREE/Source/NNERuntimeIREE/Private/NNERuntimeIREEModelData.cpp
+++ b/Engine/Plugins/Experimental/NNERuntimeIREE/Source/NNERuntimeIREE/Private/NNERuntimeIREEModelData.cpp
@@ -18,6 +18,7 @@ namespace UE::NNERuntimeIREE
 			Ar << FileId;
 			ModuleMetaData.Serialize(Ar);
 			FNNERuntimeIREECompilerResultCPU::StaticStruct()->SerializeBin(Ar, &CompilerResult);
+			Ar << CompilerHash;
 		}
 
 		bool FModelData::Load(TConstArrayView64<uint8> Data)
@@ -52,6 +53,7 @@ namespace UE::NNERuntimeIREE
 			Ar << Version;
 			Ar << FileId;
 			Ar << ShaderPlatforms;
+			Ar << CompilerHash;
 		}
 
 		void FModelData::Serialize(FArchive& Ar)
diff --git a/Engine/Plugins/Experimental/NNERuntimeIREE/Source/NNERuntimeIREE/Private/NNERuntimeIREEModelData.h b/Engine/Plugins/Experimental/NNERuntimeIREE/Source/NNERuntimeIREE/Private/NNERuntimeIREEModelData.h
--- a/Engine/Plugins/Experimental/NNERuntimeIREE/Source/NNERuntimeIREE/Private/NNERuntimeIREEModelData.h
+++ b/Engine/Plugins/Experimental/NNERuntimeIREE/Source/NNERuntimeIREE/Private/NNERuntimeIREEModelData.h
@@ -4,6 +4,7 @@
 
 #include "Containers/ContainersFwd.h"
 #include "IREECompilerRDG.h"
+#include "Misc/SecureHash.h"
 #include "NNERuntimeIREECompiler.h"
 #include "NNERuntimeIREEMetaData.h"
 #include "NNETypes.h"
@@ -60,6 +61,11 @@ namespace UE::NNERuntimeIREE::CPU::Private
 		 */
 		FNNERuntimeIREECompilerResultCPU CompilerResult;
 
+		/**
+		* Hash of the compiler that created CompilerResult
+		*/
+		FMD5Hash CompilerHash;
+
 	private:
 		void Serialize(FArchive& Ar);
 	};
@@ -94,6 +100,11 @@ namespace UE::NNERuntimeIREE::RDG::Private
 		 * List of shader platforms supported by this IREE model data.
 		 */
 		TArray<FString> ShaderPlatforms;
+
+		/**
+		* Hash of the compiler that created CompilerResult
+		*/
+		FMD5Hash CompilerHash;
 	};
 
 	/**
