[NNE] NNERuntimeIREECpu add per model runtime settings

Details:
- Added compiler settings in form of flags that get appended to compiler arguments per target and platform (first match)
- Update compiler hash with applied flags

#tests Editor, Win64, added, updated and removed to a model and verified recompilation and correct application of flags

#Synced-CL 51407547

diff --git a/Engine/Plugins/Experimental/NNERuntimeIREE/Source/NNERuntimeIREE/Private/NNERuntimeIREE.cpp b/Engine/Plugins/Experimental/NNERuntimeIREE/Source/NNERuntimeIREE/Private/NNERuntimeIREE.cpp
--- a/Engine/Plugins/Experimental/NNERuntimeIREE/Source/NNERuntimeIREE/Private/NNERuntimeIREE.cpp
+++ b/Engine/Plugins/Experimental/NNERuntimeIREE/Source/NNERuntimeIREE/Private/NNERuntimeIREE.cpp
@@ -26,6 +26,7 @@
 #include "NNERuntimeIREELog.h"
 #include "NNERuntimeIREEModel.h"
 #include "NNERuntimeIREEModelData.h"
+#include "NNERuntimeIREESettings.h"
 #include "RHIGlobals.h"
 #include "RHIStrings.h"
 #include "Serialization/MemoryReader.h"
@@ -33,6 +34,12 @@
 
 namespace UE::NNERuntimeIREE::Private
 {
+	UNNERuntimeIREECpuSettings* CreateDefaultRuntimeSettingsCPU()
+	{
+		UNNERuntimeIREECpuSettings* Settings = NewObject<UNNERuntimeIREECpuSettings>();
+		return Settings;
+	}
+
 	FString GetTargetPlatformName(const ITargetPlatform* TargetPlatform)
 	{
 		return TargetPlatform ? TargetPlatform->IniPlatformName() : UGameplayStatics::GetPlatformName();
@@ -264,7 +271,14 @@ TSharedPtr<UE::NNE::FSharedModelData> UNNERuntimeIREECpu::CreateModelData(const
 
 	FString IREEModelDataFilePath = FPaths::Combine(IntermediateDirFullPath, FileIdString) + ".ireemodeldata";
 
-	TUniquePtr<UE::NNERuntimeIREE::CPU::FCompiler> Compiler = UE::NNERuntimeIREE::CPU::FCompiler::Make(TargetPlatformName);
+	UNNERuntimeIREECpuSettings* Settings = CreateDefaultRuntimeSettingsCPU();
+	if (UNNEModelData::DeserializeRuntimeSettings(AdditionalFileData, *Settings) != UNNEModelData::EDeserializeRuntimeSettings::Success)
+	{
+		UE_LOG(LogNNERuntimeIREE, Warning, TEXT("UNNERuntimeIREECpu failed to deserialize runtime settings"));
+		return TSharedPtr<UE::NNE::FSharedModelData>();
+	}
+
+	TUniquePtr<UE::NNERuntimeIREE::CPU::FCompiler> Compiler = UE::NNERuntimeIREE::CPU::FCompiler::Make(TargetPlatformName, Settings->CompilerSettings);
 	if (!Compiler.IsValid())
 	{
 		UE_LOG(LogNNERuntimeIREE, Warning, TEXT("UNNERuntimeIREECpu failed to create a compiler to compile for platform %s"), *TargetPlatformName);
@@ -378,6 +392,11 @@ TSharedPtr<UE::NNE::FSharedModelData> UNNERuntimeIREECpu::CreateModelData(const
 #endif // WITH_EDITOR
 }
 
+UObject* UNNERuntimeIREECpu::CreateDefaultRuntimeSettings() const
+{
+	return UE::NNERuntimeIREE::Private::CreateDefaultRuntimeSettingsCPU();
+}
+
 FString UNNERuntimeIREECpu::GetModelDataIdentifier(const FString& FileType, TConstArrayView64<uint8> FileData, const TMap<FString, TConstArrayView64<uint8>>& AdditionalFileData, const FGuid& FileId, const ITargetPlatform* TargetPlatform) const
 {
 	// Leave architecture blank as there is only one model data for all architectures of a given platform, only the vmfb and shared lib are different
diff --git a/Engine/Plugins/Experimental/NNERuntimeIREE/Source/NNERuntimeIREE/Private/NNERuntimeIREE.h b/Engine/Plugins/Experimental/NNERuntimeIREE/Source/NNERuntimeIREE/Private/NNERuntimeIREE.h
--- a/Engine/Plugins/Experimental/NNERuntimeIREE/Source/NNERuntimeIREE/Private/NNERuntimeIREE.h
+++ b/Engine/Plugins/Experimental/NNERuntimeIREE/Source/NNERuntimeIREE/Private/NNERuntimeIREE.h
@@ -36,6 +36,7 @@ public:
 	virtual ECanCreateModelDataStatus CanCreateModelData(const FString& FileType, TConstArrayView64<uint8> FileData, const TMap<FString, TConstArrayView64<uint8>>& AdditionalFileData, const FGuid& FileId, const ITargetPlatform* TargetPlatform) const override;
 	virtual TSharedPtr<UE::NNE::FSharedModelData> CreateModelData(const FString& FileType, TConstArrayView64<uint8> FileData, const TMap<FString, TConstArrayView64<uint8>>& AdditionalFileData, const FGuid& FileId, const ITargetPlatform* TargetPlatform) override;
 	virtual FString GetModelDataIdentifier(const FString& FileType, TConstArrayView64<uint8> FileData, const TMap<FString, TConstArrayView64<uint8>>& AdditionalFileData, const FGuid& FileId, const ITargetPlatform* TargetPlatform) const override;
+	virtual UObject* CreateDefaultRuntimeSettings() const override;
 	//~ End INNERuntime Interface
 
 	//~ Begin INNERuntimeCPU Interface
diff --git a/Engine/Plugins/Experimental/NNERuntimeIREE/Source/NNERuntimeIREE/Private/NNERuntimeIREECompiler.cpp b/Engine/Plugins/Experimental/NNERuntimeIREE/Source/NNERuntimeIREE/Private/NNERuntimeIREECompiler.cpp
--- a/Engine/Plugins/Experimental/NNERuntimeIREE/Source/NNERuntimeIREE/Private/NNERuntimeIREECompiler.cpp
+++ b/Engine/Plugins/Experimental/NNERuntimeIREE/Source/NNERuntimeIREE/Private/NNERuntimeIREECompiler.cpp
@@ -1,6 +1,9 @@
 // Copyright Epic Games, Inc. All Rights Reserved.
 
 #include "NNERuntimeIREECompiler.h"
+#include "Logging/LogVerbosity.h"
+#include "Misc/CString.h"
+#include "Templates/UnrealTemplate.h"
 
 #ifdef WITH_NNE_RUNTIME_IREE
 #if WITH_EDITOR
@@ -17,6 +20,7 @@
 #include "Misc/Paths.h"
 #include "NNE.h"
 #include "NNERuntimeIREELog.h"
+#include "NNERuntimeIREESettings.h"
 #include "Serialization/JsonReader.h"
 #include "Serialization/JsonSerializer.h"
 #include "Templates/SharedPointer.h"
@@ -51,7 +55,7 @@ namespace UE::NNERuntimeIREE
 
 		}
 
-		TUniquePtr<FCompiler> FCompiler::Make(const FString& InTargetPlatformName)
+		TUniquePtr<FCompiler> FCompiler::Make(const FString& InTargetPlatformName, const FNNERuntimeIREECpuCompilerSettings& CompilerSettings)
 		{
 			SCOPED_NAMED_EVENT_TEXT("FCompiler::Make", FColor::Magenta);
 
@@ -163,16 +167,50 @@ namespace UE::NNERuntimeIREE
 								continue;
 							}
 
+							TArray<FBuildTarget> TmpBuildTargets = BuildConfig.BuildTargets;
+							FString SettingsToHash;
+							for (FBuildTarget& Target : TmpBuildTargets)
+							{
+								for (const FNNERuntimeIREECpuCompilerFlags& FlagSettings: CompilerSettings.CompilerFlags)
+								{
+									if (!FlagSettings.Platform.IsEmpty() && FlagSettings.Platform.Compare(InTargetPlatformName, ESearchCase::IgnoreCase) != 0)
+									{
+										continue;
+									}
+
+									if (!FlagSettings.Target.IsEmpty() && FlagSettings.Target.Compare(Target.Name, ESearchCase::IgnoreCase) != 0)
+									{
+										continue;
+									}
+
+									Target.CompilerArguments.AppendChar(' ');
+									Target.CompilerArguments.Append(FlagSettings.Flags);
+
+									UE_LOG(LogNNERuntimeIREE, Verbose, TEXT("Updated IREE compiler flags for %s on %s: %s"), *Target.Name, *InTargetPlatformName, *Target.CompilerArguments);
+
+									SettingsToHash.Append(FlagSettings.Platform)
+										.AppendChar('|')
+										.Append(FlagSettings.Target)
+										.AppendChar('|')
+										.Append(FlagSettings.Flags);
+
+									break;
+								}
+
+								SettingsToHash.AppendChar('\n');
+							}
+
 							ImporterCommand = TmpImporterCommand;
 							ImporterArguments = BuildConfig.ImporterArguments;
 							CompilerCommand = TmpCompilerCommand;
 							LinkerCommand = TmpLinkerCommand;
 							SharedLibExt = BuildConfig.SharedLibExt;
-							BuildTargets = BuildConfig.BuildTargets;
+							BuildTargets = MoveTemp(TmpBuildTargets);
 
 							FMD5 HashBuilder{};
 							bool bHashSuccess = true;
 							UE::IREEUtils::HashAppendString(HashBuilder, BuildConfigFileString);
+							UE::IREEUtils::HashAppendString(HashBuilder, SettingsToHash);
 							bHashSuccess &= UE::IREEUtils::HashAppendFileStat(HashBuilder, ImporterCommand);
 							bHashSuccess &= UE::IREEUtils::HashAppendFileStat(HashBuilder, CompilerCommand);
 							bHashSuccess &= UE::IREEUtils::HashAppendFileStat(HashBuilder, LinkerCommand);
diff --git a/Engine/Plugins/Experimental/NNERuntimeIREE/Source/NNERuntimeIREE/Private/NNERuntimeIREECompiler.h b/Engine/Plugins/Experimental/NNERuntimeIREE/Source/NNERuntimeIREE/Private/NNERuntimeIREECompiler.h
--- a/Engine/Plugins/Experimental/NNERuntimeIREE/Source/NNERuntimeIREE/Private/NNERuntimeIREECompiler.h
+++ b/Engine/Plugins/Experimental/NNERuntimeIREE/Source/NNERuntimeIREE/Private/NNERuntimeIREECompiler.h
@@ -8,6 +8,8 @@
 
 #include "NNERuntimeIREECompiler.generated.h"
 
+struct FNNERuntimeIREECpuCompilerSettings;
+
 USTRUCT()
 struct FNNERuntimeIREEArchitectureInfoCPU
 {
@@ -94,7 +96,7 @@ namespace UE::NNERuntimeIREE::CPU
 
 	public:
 		~FCompiler() {};
-		static TUniquePtr<FCompiler> Make(const FString& InTargetPlatformName);
+		static TUniquePtr<FCompiler> Make(const FString& InTargetPlatformName, const FNNERuntimeIREECpuCompilerSettings& CompilerSettings);
 
 	public:
 		const FMD5Hash& GetHash() const { return Hash; }
diff --git a/Engine/Plugins/Experimental/NNERuntimeIREE/Source/NNERuntimeIREE/Private/NNERuntimeIREESettings.cpp b/Engine/Plugins/Experimental/NNERuntimeIREE/Source/NNERuntimeIREE/Private/NNERuntimeIREESettings.cpp
--- a/Engine/Plugins/Experimental/NNERuntimeIREE/Source/NNERuntimeIREE/Private/NNERuntimeIREESettings.cpp
+++ b/Engine/Plugins/Experimental/NNERuntimeIREE/Source/NNERuntimeIREE/Private/NNERuntimeIREESettings.cpp
@@ -2,6 +2,9 @@
 
 #include "NNERuntimeIREESettings.h"
 
+#include "Interfaces/ITargetPlatformManagerModule.h"
+#include "Interfaces/ITargetPlatform.h"
+
 UNNERuntimeIREESettings::UNNERuntimeIREESettings(const FObjectInitializer& ObjectInitlaizer)
 	: Super(ObjectInitlaizer)
 {
diff --git a/Engine/Plugins/Experimental/NNERuntimeIREE/Source/NNERuntimeIREE/Private/NNERuntimeIREESettings.h b/Engine/Plugins/Experimental/NNERuntimeIREE/Source/NNERuntimeIREE/Private/NNERuntimeIREESettings.h
--- a/Engine/Plugins/Experimental/NNERuntimeIREE/Source/NNERuntimeIREE/Private/NNERuntimeIREESettings.h
+++ b/Engine/Plugins/Experimental/NNERuntimeIREE/Source/NNERuntimeIREE/Private/NNERuntimeIREESettings.h
@@ -3,6 +3,7 @@
 #pragma once
 
 #include "Engine/DeveloperSettings.h"
+#include "UObject/ObjectMacros.h"
 
 #include "NNERuntimeIREESettings.generated.h"
 
@@ -99,3 +100,58 @@ public:
 #endif
 
 };
+
+/** NNERuntimeIREECpu compiler flag settings */
+USTRUCT()
+struct FNNERuntimeIREECpuCompilerFlags
+{
+	GENERATED_BODY()
+
+	/** The Platform the flags should be applied to, or leave empty to apply to all */
+	UPROPERTY(EditAnywhere, Category = "IREE Runtime", meta = (
+		DisplayName = "Platform",
+		ToolTip = "The Platform the flags should be applied to, or leave empty to apply to all."))
+	FString Platform;
+
+	/** The name of the target these flags should be applied to, or leave empty to apply to all */
+	UPROPERTY(EditAnywhere, Category = "IREE Runtime", meta = (
+		DisplayName = "Target",
+		ToolTip = "The name of the target these flags should be applied to, or leave empty to apply to all."))
+	FString Target;
+
+	/** The flags that get appended to the iree compiler arguments on the specified platform and target */
+	UPROPERTY(EditAnywhere, Category = "IREE Runtime", meta = (
+		DisplayName = "Flags",
+		ToolTip = "The flags that get appended to the iree compiler arguments on the specified platform and target."))
+	FString Flags;
+};
+
+/** NNERuntimeIREECpu compiler settings */
+USTRUCT()
+struct FNNERuntimeIREECpuCompilerSettings
+{
+	GENERATED_BODY()
+
+	/** NNERuntimeIREECpu compiler flag settings */
+	UPROPERTY(EditAnywhere, Category = "IREE Runtime", meta = (
+		DisplayName = "Compiler Flags", 
+		ToolTip = "List of compiler flags. The flags of the first item that has the correct platform and target get appended to the existing iree compiler arguments.",
+		TitleProperty="{Platform}, {Target}: {Flags}"))
+	TArray<FNNERuntimeIREECpuCompilerFlags> CompilerFlags;
+};
+
+/** RuntimeSettings for NNERuntimeIREECpu */
+UCLASS()
+class UNNERuntimeIREECpuSettings : public UObject
+{
+	GENERATED_BODY()
+
+public:
+
+	/** NNERuntimeIREECpu compiler settings */
+	UPROPERTY(EditAnywhere, Category = "IREE Runtime", meta = (
+		DisplayName = "Compiler settings", 
+		ToolTip = "Settings to configure compilation.",
+		NoResetToDefault))
+	FNNERuntimeIREECpuCompilerSettings CompilerSettings;
+};
