[NNE] Fix NNERuntimeIREERDG error when runing previously compiled models

#rnx
#tests Loading of an old IREERDG model

A change in the way shaders are serialized means that previously compiled models throw an "Mismatch size for type FNNERuntimeIREEShader" error. This fix simply increases the version number of IREE RDG to ensure that the shaders get recompiled to the new format.

#Synced-CL 52167502

diff --git a/Engine/Plugins/Experimental/NNERuntimeIREE/Source/NNERuntimeIREE/Private/NNERuntimeIREE.cpp b/Engine/Plugins/Experimental/NNERuntimeIREE/Source/NNERuntimeIREE/Private/NNERuntimeIREE.cpp
--- a/Engine/Plugins/Experimental/NNERuntimeIREE/Source/NNERuntimeIREE/Private/NNERuntimeIREE.cpp
+++ b/Engine/Plugins/Experimental/NNERuntimeIREE/Source/NNERuntimeIREE/Private/NNERuntimeIREE.cpp
@@ -622,7 +622,7 @@ int32 UNNERuntimeIREEVulkan::GetVersion() const
 #ifdef WITH_NNE_RUNTIME_IREE_RDG
 
 FGuid UNNERuntimeIREERdg::GUID = FGuid((int32)'I', (int32)'R', (int32)'D', (int32)'G');
-int32 UNNERuntimeIREERdg::Version = 0x00000005; // Increment each time iree compiler/runtime is updated and each time FModelData serialization changes
+int32 UNNERuntimeIREERdg::Version = 0x00000006; // Increment each time iree compiler/runtime is updated and each time FModelData serialization changes
 
 FString UNNERuntimeIREERdg::GetRuntimeName() const
 {
diff --git a/Engine/Plugins/Experimental/NNERuntimeIREE/Source/NNERuntimeIREE/Private/NNERuntimeIREEModel.cpp b/Engine/Plugins/Experimental/NNERuntimeIREE/Source/NNERuntimeIREE/Private/NNERuntimeIREEModel.cpp
--- a/Engine/Plugins/Experimental/NNERuntimeIREE/Source/NNERuntimeIREE/Private/NNERuntimeIREEModel.cpp
+++ b/Engine/Plugins/Experimental/NNERuntimeIREE/Source/NNERuntimeIREE/Private/NNERuntimeIREEModel.cpp
@@ -370,7 +370,7 @@ namespace UE::NNERuntimeIREE
 				Status = iree_runtime_session_create_with_device(Instance, &SessionOptions, InDevice, GetHostAllocator(), &Session);
 				if (!iree_status_is_ok(Status))
 				{
-					Private::PrintIREEError("CPU instance: Failed to create the session", Status);
+					Private::PrintIREEError("Failed to create the session", Status);
 
 					if (Session)
 					{
